[Mass] Added a convenience function for showing consistent editor notifies for Mass

diff --git a/Engine/Source/Editor/MassEntityEditor/Private/MassEntityEditor.cpp b/Engine/Source/Editor/MassEntityEditor/Private/MassEntityEditor.cpp
new file mode 100644
--- /dev/null
+++ b/Engine/Source/Editor/MassEntityEditor/Private/MassEntityEditor.cpp
@@ -0,0 +1,32 @@
+// Copyright Epic Games, Inc. All Rights Reserved.
+
+#include "MassEntityEditor.h"
+#include "Logging/MessageLog.h"
+#include "Framework/Docking/TabManager.h"
+
+#define LOCTEXT_NAMESPACE "Mass"
+
+void FMassEditorNotification::Show()
+{
+	FMessageLog MessageLog(UE::Mass::Editor::MessageLogPageName);
+
+	MessageLog.AddMessage(FTokenizedMessage::Create(Severity))
+		->AddToken(FTextToken::Create(Message));
+
+	if (bIncludeSeeOutputLogForDetails)
+	{
+		MessageLog.AddMessage(FTokenizedMessage::Create(EMessageSeverity::Info))
+			->AddToken(FActionToken::Create(LOCTEXT("MassSeeLogForDetails", "See the log for more details.")
+				, LOCTEXT("MassSeeLogForDetailsTooltip", "Open the Output Log tab.")
+				, FOnActionTokenExecuted::CreateLambda([]()
+					{
+						FGlobalTabmanager::Get()->TryInvokeTab(FName("OutputLog"));
+					}))
+			);
+	}
+
+	// Forcing so that even the "Info"-level notifications get shown
+	MessageLog.Notify(Message, EMessageSeverity::Info, /*bForce=*/true);
+}
+
+#undef LOCTEXT_NAMESPACE
diff --git a/Engine/Source/Editor/MassEntityEditor/Public/MassEntityEditor.h b/Engine/Source/Editor/MassEntityEditor/Public/MassEntityEditor.h
new file mode 100644
--- /dev/null
+++ b/Engine/Source/Editor/MassEntityEditor/Public/MassEntityEditor.h
@@ -0,0 +1,29 @@
+// Copyright Epic Games, Inc. All Rights Reserved.
+
+#pragma once
+
+#include "Logging/TokenizedMessage.h"
+#include "MassEntityEditor.generated.h"
+
+
+USTRUCT()
+struct MASSENTITYEDITOR_API FMassEditorNotification
+{
+	GENERATED_BODY()
+
+	EMessageSeverity::Type Severity = EMessageSeverity::Error;
+	FText Message;
+
+	/** 
+	 * If set to true then a clickable "see log for details" message will be added to the message log. 
+	 * Clicking the message takes the user to the Output Log.
+	 */
+	bool bIncludeSeeOutputLogForDetails = false;
+
+	void Show();
+};
+
+namespace UE::Mass::Editor
+{
+	const FName MessageLogPageName(TEXT("MassEntity"));
+}
\ No newline at end of file
