[Mass] Unified naming of MassEntity page in MessageLog

diff --git a/Engine/Source/Editor/MassEntityEditor/Private/MassEntityEditorModule.cpp b/Engine/Source/Editor/MassEntityEditor/Private/MassEntityEditorModule.cpp
--- a/Engine/Source/Editor/MassEntityEditor/Private/MassEntityEditorModule.cpp
+++ b/Engine/Source/Editor/MassEntityEditor/Private/MassEntityEditorModule.cpp
@@ -7,9 +7,11 @@
 #include "MessageLogModule.h"
 #include "Engine/World.h"
 #include "Logging/MessageLog.h"
+#include "MassDebugger.h"
+#include "MassEntityEditor.h"
 #endif // WITH_UNREAL_DEVELOPER_TOOLS
 
-#define LOCTEXT_NAMESPACE "MassEntityEditor"
+#define LOCTEXT_NAMESPACE "Mass"
 
 IMPLEMENT_MODULE(FMassEntityEditorModule, MassEntityEditor)
 
@@ -25,7 +27,8 @@ void FMassEntityEditorModule::StartupModule()
 	FMessageLogInitializationOptions InitOptions;
 	InitOptions.bShowPages = true;
 	InitOptions.bShowFilters = true;
-	MessageLogModule.RegisterLogListing("MassEntity", LOCTEXT("MassEntity", "MassEntity"), InitOptions);
+	MessageLogModule.RegisterLogListing(UE::Mass::Editor::MessageLogPageName
+		, FText::FromName(UE::Mass::Editor::MessageLogPageName), InitOptions);
 
 	OnWorldCleanupHandle = FWorldDelegates::OnWorldCleanup.AddStatic(&FMassEntityEditorModule::OnWorldCleanup);
 
@@ -50,7 +53,7 @@ void FMassEntityEditorModule::ShutdownModule()
 void FMassEntityEditorModule::OnWorldCleanup(UWorld* /*World*/, bool /*bSessionEnded*/, bool /*bCleanupResources*/)
 {
 	// clearing out messages from the world being cleaned up
-	FMessageLog("MassEntity").NewPage(FText::FromString(TEXT("MassEntity")));
+	FMessageLog(UE::Mass::Editor::MessageLogPageName).NewPage(FText::FromName(UE::Mass::Editor::MessageLogPageName));
 }
 #endif // WITH_UNREAL_DEVELOPER_TOOLS
 
