Cleanup of FSlateFontInfo constructor across the solution that uses font paths. It will be deprecated soon.

#rnx

diff --git a/Engine/Plugins/Animation/ControlRig/Source/ControlRigEditor/Private/ControlRigEditorStyle.h b/Engine/Plugins/Animation/ControlRig/Source/ControlRigEditor/Private/ControlRigEditorStyle.h
--- a/Engine/Plugins/Animation/ControlRig/Source/ControlRigEditor/Private/ControlRigEditorStyle.h
+++ b/Engine/Plugins/Animation/ControlRig/Source/ControlRigEditor/Private/ControlRigEditorStyle.h
@@ -2,6 +2,7 @@
 
 #pragma once
 
+#include "Styling/CoreStyle.h"
 #include "Styling/SlateStyle.h"
 #include "Styling/SlateStyleRegistry.h"
 #include "Styling/SlateStyleMacros.h"
@@ -9,9 +10,6 @@
 #include "Styling/StyleColors.h"
 #include "Brushes/SlateRoundedBoxBrush.h"
 
-#define TTF_FONT(RelativePath, ...) FSlateFontInfo(RootToContentDir(RelativePath, TEXT(".ttf")), __VA_ARGS__)
-#define OTF_FONT(RelativePath, ...) FSlateFontInfo(RootToContentDir(RelativePath, TEXT(".otf")), __VA_ARGS__)
-
 class FControlRigEditorStyle final
 	: public FSlateStyleSet
 {
@@ -141,7 +139,7 @@ public:
 
 		// Font?
 		{
-			Set("ControlRig.Hierarchy.Menu", TTF_FONT("Fonts/Roboto-Regular", 12));
+			Set("ControlRig.Hierarchy.Menu", DEFAULT_FONT("Regular", 12));
 		}
 
 		// Space picker
@@ -213,6 +211,3 @@ public:
 	FLinearColor ConnectorUserInterfaceColor;
 	FLinearColor SocketUserInterfaceColor;
 };
-
-#undef TTF_FONT
-#undef OTF_FONT
diff --git a/Engine/Plugins/Animation/GameplayInsights/Source/GameplayInsights/Private/GameplayInsightsStyle.h b/Engine/Plugins/Animation/GameplayInsights/Source/GameplayInsights/Private/GameplayInsightsStyle.h
--- a/Engine/Plugins/Animation/GameplayInsights/Source/GameplayInsights/Private/GameplayInsightsStyle.h
+++ b/Engine/Plugins/Animation/GameplayInsights/Source/GameplayInsights/Private/GameplayInsightsStyle.h
@@ -7,9 +7,6 @@
 #include "Styling/SlateStyleMacros.h"
 #include "Styling/SlateTypes.h"
 
-#define TTF_FONT(RelativePath, ...) FSlateFontInfo(RootToContentDir(RelativePath, TEXT(".ttf")), __VA_ARGS__)
-#define OTF_FONT(RelativePath, ...) FSlateFontInfo(RootToContentDir(RelativePath, TEXT(".otf")), __VA_ARGS__)
-
 class FGameplayInsightsStyle final
 	: public FSlateStyleSet
 {
@@ -75,9 +72,3 @@ public:
 	const FSlateColor SelectionColor;
 	const FSlateColor SelectionColor_Pressed;
 };
-
-#undef IMAGE_BRUSH
-#undef BOX_BRUSH
-#undef BORDER_BRUSH
-#undef TTF_FONT
-#undef OTF_FONT
diff --git a/Engine/Plugins/Developer/TraceDataFiltering/Source/TraceDataFiltering/Private/EventFilterStyle.cpp b/Engine/Plugins/Developer/TraceDataFiltering/Source/TraceDataFiltering/Private/EventFilterStyle.cpp
--- a/Engine/Plugins/Developer/TraceDataFiltering/Source/TraceDataFiltering/Private/EventFilterStyle.cpp
+++ b/Engine/Plugins/Developer/TraceDataFiltering/Source/TraceDataFiltering/Private/EventFilterStyle.cpp
@@ -20,8 +20,6 @@ TSharedPtr< FSlateStyleSet > FEventFilterStyle::StyleSet = nullptr;
 
 FTextBlockStyle FEventFilterStyle::NormalText;
 
-#define ICON_FONT(...) FSlateFontInfo(RootToContentDir("Fonts/FontAwesome", TEXT(".ttf")), __VA_ARGS__)
-
 // Const icon sizes
 static const FVector2D Icon8x8(8.0f, 8.0f);
 static const FVector2D Icon9x19(9.0f, 19.0f);
@@ -133,16 +131,13 @@ void FEventFilterStyle::Initialize()
 	StyleSet->Set("FilterPresets.SessionWarningBorder", new BOX_BRUSH("Common/DarkGroupBorder", FMargin(4.0f / 16.0f)));
 	StyleSet->Set("FilterPresets.WarningIcon", new CORE_IMAGE_BRUSH_SVG("Starship/Common/alert-circle", Icon40x40, FStyleColors::Warning));
 
-	StyleSet->Set("FontAwesome.9", ICON_FONT(9));
+	StyleSet->Set("FontAwesome.9", REGULAR_ICON_FONT(9));
 	
 	FSlateStyleRegistry::RegisterSlateStyle( *StyleSet.Get() );
 }
 
-#undef IMAGE_BRUSH
-#undef BOX_BRUSH
-#undef BORDER_BRUSH
-#undef DEFAULT_FONT
-#undef ICON_FONT
+#undef RootToContentDir
+#undef RootToCoreContentDir
 
 void FEventFilterStyle::Shutdown()
 {
diff --git a/Engine/Plugins/Developer/TraceSourceFiltering/Source/SourceFilteringEditor/Private/SourceFilterStyle.cpp b/Engine/Plugins/Developer/TraceSourceFiltering/Source/SourceFilteringEditor/Private/SourceFilterStyle.cpp
--- a/Engine/Plugins/Developer/TraceSourceFiltering/Source/SourceFilteringEditor/Private/SourceFilterStyle.cpp
+++ b/Engine/Plugins/Developer/TraceSourceFiltering/Source/SourceFilteringEditor/Private/SourceFilterStyle.cpp
@@ -21,9 +21,6 @@ TSharedPtr< FSlateStyleSet > FSourceFilterStyle::StyleSet = nullptr;
 
 FTextBlockStyle FSourceFilterStyle::NormalText;
 
-
-#define ICON_FONT(...) FSlateFontInfo(RootToContentDir("Fonts/FontAwesome", TEXT(".ttf")), __VA_ARGS__)
-
 // Const icon sizes
 static const FVector2D Icon8x8(8.0f, 8.0f);
 static const FVector2D Icon9x19(9.0f, 19.0f);
@@ -168,17 +165,13 @@ void FSourceFilterStyle::Initialize()
 
 	StyleSet->Set("SourceFilter.DragDrop.Border", new BOX_BRUSH("Old/Menu_Background", FMargin(8.0f / 64.0f)));
 
-	StyleSet->Set("FontAwesome.9", ICON_FONT(9));
-	StyleSet->Set("FontAwesome.12", ICON_FONT(12));
+	StyleSet->Set("FontAwesome.9", REGULAR_ICON_FONT(9));
+	StyleSet->Set("FontAwesome.12", REGULAR_ICON_FONT(12));
 		
 	FSlateStyleRegistry::RegisterSlateStyle( *StyleSet.Get() );
 }
 
-#undef IMAGE_BRUSH
-#undef BOX_BRUSH
-#undef BORDER_BRUSH
-#undef DEFAULT_FONT
-#undef ICON_FONT
+#undef RootToContentDir
 
 void FSourceFilterStyle::Shutdown()
 {
diff --git a/Engine/Plugins/Editor/AssetManagerEditor/Source/AssetManagerEditor/Private/ReferenceViewer/ReferenceViewerStyle.cpp b/Engine/Plugins/Editor/AssetManagerEditor/Source/AssetManagerEditor/Private/ReferenceViewer/ReferenceViewerStyle.cpp
--- a/Engine/Plugins/Editor/AssetManagerEditor/Source/AssetManagerEditor/Private/ReferenceViewer/ReferenceViewerStyle.cpp
+++ b/Engine/Plugins/Editor/AssetManagerEditor/Source/AssetManagerEditor/Private/ReferenceViewer/ReferenceViewerStyle.cpp
@@ -14,8 +14,6 @@
 
 FName FReferenceViewerStyle::StyleName("ReferenceViewerStyle");
 
-#define FONT(...) FSlateFontInfo(FCoreStyle::GetDefaultFont(), __VA_ARGS__)
-
 FReferenceViewerStyle::FReferenceViewerStyle()
 	: FSlateStyleSet(StyleName)
 {
@@ -28,7 +26,7 @@ FReferenceViewerStyle::FReferenceViewerStyle()
 	FTextBlockStyle NormalText = GetParentStyle()->GetWidgetStyle<FTextBlockStyle>("NormalText");
 	FEditableTextBoxStyle NormalEditableTextBoxStyle = GetParentStyle()->GetWidgetStyle<FEditableTextBoxStyle>("NormalEditableTextBox");
 
-	FSlateFontInfo TitleFont = FONT(12, "Bold");
+	FSlateFontInfo TitleFont = DEFAULT_FONT("Bold", 12);
 
 	// Text Styles 
 	FTextBlockStyle GraphNodeTitle = FTextBlockStyle(NormalText)
diff --git a/Engine/Plugins/Editor/GameplayTagsEditor/Source/GameplayTagsEditor/Private/GameplayTagStyle.cpp b/Engine/Plugins/Editor/GameplayTagsEditor/Source/GameplayTagsEditor/Private/GameplayTagStyle.cpp
--- a/Engine/Plugins/Editor/GameplayTagsEditor/Source/GameplayTagsEditor/Private/GameplayTagStyle.cpp
+++ b/Engine/Plugins/Editor/GameplayTagsEditor/Source/GameplayTagsEditor/Private/GameplayTagStyle.cpp
@@ -9,9 +9,6 @@
 #include "Components/Widget.h"
 #include "Styling/StyleColors.h"
 
-#define TTF_CORE_FONT(RelativePath, ...) FSlateFontInfo(RootToCoreContentDir(RelativePath, TEXT(".ttf")), __VA_ARGS__)
-#define FONT(...) FSlateFontInfo(FCoreStyle::GetDefaultFont(), __VA_ARGS__)
-
 class FGameplayTagStyleSet final : public FSlateStyleSet
 {
 public:
diff --git a/Engine/Plugins/Editor/GeometryMode/Source/BspMode/Private/BspModeStyle.cpp b/Engine/Plugins/Editor/GeometryMode/Source/BspMode/Private/BspModeStyle.cpp
--- a/Engine/Plugins/Editor/GeometryMode/Source/BspMode/Private/BspModeStyle.cpp
+++ b/Engine/Plugins/Editor/GeometryMode/Source/BspMode/Private/BspModeStyle.cpp
@@ -10,15 +10,12 @@
 #include "Brushes/SlateNoResource.h"
 #include "Misc/Paths.h"
 #include "Styling/SlateStyle.h"
+#include "Styling/SlateStyleMacros.h"
 
 
 TSharedPtr< FSlateStyleSet > FBspModeStyle::StyleSet = nullptr;
 
-#define IMAGE_BRUSH( RelativePath, ... ) FSlateImageBrush( StyleSet->RootToContentDir( RelativePath, TEXT(".png") ), __VA_ARGS__ )
-#define BOX_BRUSH( RelativePath, ... ) FSlateBoxBrush( StyleSet->RootToContentDir( RelativePath, TEXT(".png") ), __VA_ARGS__ )
-#define BORDER_BRUSH( RelativePath, ... ) FSlateBorderBrush( StyleSet->RootToContentDir( RelativePath, TEXT(".png") ), __VA_ARGS__ )
-#define DEFAULT_FONT(...) FCoreStyle::GetDefaultFontStyle(__VA_ARGS__)
-
+#define RootToContentDir StyleSet->RootToContentDir
 
 void FBspModeStyle::Initialize()
 {
@@ -125,10 +122,7 @@ void FBspModeStyle::Initialize()
 	FSlateStyleRegistry::RegisterSlateStyle( *StyleSet.Get() );
 }
 
-#undef IMAGE_BRUSH
-#undef BOX_BRUSH
-#undef BORDER_BRUSH
-#undef DEFAULT_FONT
+#undef RootToContentDir
 
 void FBspModeStyle::Shutdown()
 {
diff --git a/Engine/Plugins/Editor/GeometryMode/Source/GeometryMode/Private/GeometryModeStyle.cpp b/Engine/Plugins/Editor/GeometryMode/Source/GeometryMode/Private/GeometryModeStyle.cpp
--- a/Engine/Plugins/Editor/GeometryMode/Source/GeometryMode/Private/GeometryModeStyle.cpp
+++ b/Engine/Plugins/Editor/GeometryMode/Source/GeometryMode/Private/GeometryModeStyle.cpp
@@ -8,12 +8,10 @@
 #include "Styling/AppStyle.h"
 #include "Interfaces/IPluginManager.h"
 #include "Styling/SlateStyle.h"
-
+#include "Styling/SlateStyleMacros.h"
 
 #define IMAGE_PLUGIN_BRUSH( RelativePath, ... ) FSlateImageBrush( FGeometryModeStyle::InContent( RelativePath, ".png" ), __VA_ARGS__ )
-#define IMAGE_BRUSH(RelativePath, ...) FSlateImageBrush(StyleSet->RootToContentDir(RelativePath, TEXT(".png")), __VA_ARGS__)
-#define BOX_BRUSH(RelativePath, ...) FSlateBoxBrush(StyleSet->RootToContentDir(RelativePath, TEXT(".png")), __VA_ARGS__)
-#define DEFAULT_FONT(...) FCoreStyle::GetDefaultFontStyle(__VA_ARGS__)
+#define RootToContentDir StyleSet->RootToContentDir
 
 FString FGeometryModeStyle::InContent(const FString& RelativePath, const ANSICHAR* Extension)
 {
@@ -72,11 +70,8 @@ void FGeometryModeStyle::Initialize()
 	FSlateStyleRegistry::RegisterSlateStyle(*StyleSet.Get());
 };
 
-
 #undef IMAGE_PLUGIN_BRUSH
-#undef IMAGE_BRUSH
-#undef BOX_BRUSH
-#undef DEFAULT_FONT
+#undef RootToContentDir
 
 void FGeometryModeStyle::Shutdown()
 {
diff --git a/Engine/Plugins/Editor/ModelingToolsEditorMode/Source/ModelingToolsEditorMode/Private/ModelingToolsEditorModeStyle.cpp b/Engine/Plugins/Editor/ModelingToolsEditorMode/Source/ModelingToolsEditorMode/Private/ModelingToolsEditorModeStyle.cpp
--- a/Engine/Plugins/Editor/ModelingToolsEditorMode/Source/ModelingToolsEditorMode/Private/ModelingToolsEditorModeStyle.cpp
+++ b/Engine/Plugins/Editor/ModelingToolsEditorMode/Source/ModelingToolsEditorMode/Private/ModelingToolsEditorModeStyle.cpp
@@ -13,13 +13,9 @@
 #include "Styling/StarshipCoreStyle.h"
 #include "Styling/StyleColors.h"
 
-
 #define IMAGE_PLUGIN_BRUSH( RelativePath, ... ) FSlateImageBrush( FModelingToolsEditorModeStyle::InContent( RelativePath, ".png" ), __VA_ARGS__ )
-
-// This is to fix the issue that SlateStyleMacros like IMAGE_BRUSH look for RootToContentDir but StyleSet->RootToContentDir is how this style is set up
 #define RootToContentDir StyleSet->RootToContentDir
 
-
 FString FModelingToolsEditorModeStyle::InContent(const FString& RelativePath, const ANSICHAR* Extension)
 {
 	static FString ContentDir = IPluginManager::Get().FindPlugin(TEXT("ModelingToolsEditorMode"))->GetContentDir();
@@ -548,9 +544,7 @@ void FModelingToolsEditorModeStyle::Initialize()
 END_SLATE_FUNCTION_BUILD_OPTIMIZATION
 
 #undef IMAGE_PLUGIN_BRUSH
-#undef IMAGE_BRUSH
-#undef BOX_BRUSH
-#undef DEFAULT_FONT
+#undef RootToContentDir
 
 void FModelingToolsEditorModeStyle::Shutdown()
 {
diff --git a/Engine/Plugins/Editor/ScriptableToolsEditorMode/Source/ScriptableToolsEditorMode/Private/ScriptableToolsEditorModeStyle.cpp b/Engine/Plugins/Editor/ScriptableToolsEditorMode/Source/ScriptableToolsEditorMode/Private/ScriptableToolsEditorModeStyle.cpp
--- a/Engine/Plugins/Editor/ScriptableToolsEditorMode/Source/ScriptableToolsEditorMode/Private/ScriptableToolsEditorModeStyle.cpp
+++ b/Engine/Plugins/Editor/ScriptableToolsEditorMode/Source/ScriptableToolsEditorMode/Private/ScriptableToolsEditorModeStyle.cpp
@@ -11,11 +11,8 @@
 #include "Styling/SlateStyle.h"
 #include "Styling/StyleColors.h"
 
-
 #define IMAGE_PLUGIN_BRUSH( RelativePath, ... ) FSlateImageBrush( FScriptableToolsEditorModeStyle::InContent( RelativePath, ".png" ), __VA_ARGS__ )
-#define IMAGE_BRUSH_SVG( RelativePath, ... ) FSlateVectorImageBrush(FScriptableToolsEditorModeStyle::InContent( RelativePath, ".svg"), __VA_ARGS__)
-
-// This is to fix the issue that SlateStyleMacros like IMAGE_BRUSH look for RootToContentDir but StyleSet->RootToContentDir is how this style is set up
+#define IMAGE_PLUGIN_BRUSH_SVG( RelativePath, ... ) FSlateVectorImageBrush(FScriptableToolsEditorModeStyle::InContent( RelativePath, ".svg"), __VA_ARGS__)
 #define RootToContentDir StyleSet->RootToContentDir
 
 FString FScriptableToolsEditorModeStyle::InContent(const FString& RelativePath, const ANSICHAR* Extension)
@@ -71,7 +68,7 @@ void FScriptableToolsEditorModeStyle::Initialize()
 
 	}
 
-	StyleSet->Set("ToolPalette.MenuIndicator", new IMAGE_BRUSH_SVG("Icons/chevron-right", Icon20x20, FStyleColors::Foreground));
+	StyleSet->Set("ToolPalette.MenuIndicator", new IMAGE_PLUGIN_BRUSH_SVG("Icons/chevron-right", Icon20x20, FStyleColors::Foreground));
 
 	FSlateStyleRegistry::RegisterSlateStyle(*StyleSet.Get());
 };
@@ -123,14 +120,11 @@ bool FScriptableToolsEditorModeStyle::TryRegisterCustomIcon(FName StyleIdentifie
 	return true;
 }
 
-
-
 END_SLATE_FUNCTION_BUILD_OPTIMIZATION
 
 #undef IMAGE_PLUGIN_BRUSH
-#undef IMAGE_BRUSH
-#undef BOX_BRUSH
-#undef DEFAULT_FONT
+#undef IMAGE_PLUGIN_BRUSH_SVG
+#undef RootToContentDir
 
 void FScriptableToolsEditorModeStyle::Shutdown()
 {
diff --git a/Engine/Plugins/Enterprise/DataprepEditor/Source/DataprepEditor/Private/DataprepEditorStyle.cpp b/Engine/Plugins/Enterprise/DataprepEditor/Source/DataprepEditor/Private/DataprepEditorStyle.cpp
--- a/Engine/Plugins/Enterprise/DataprepEditor/Source/DataprepEditor/Private/DataprepEditorStyle.cpp
+++ b/Engine/Plugins/Enterprise/DataprepEditor/Source/DataprepEditor/Private/DataprepEditorStyle.cpp
@@ -10,11 +10,10 @@
 #include "Misc/Paths.h"
 #include "Styling/SlateStyleRegistry.h"
 #include "Styling/ToolBarStyle.h"
+#include "Styling/SlateStyleMacros.h"
 
 #define IMAGE_PLUGIN_BRUSH( RelativePath, ... ) FSlateImageBrush( FDataprepEditorStyle::InContent( RelativePath, ".png" ), __VA_ARGS__ )
-#define BOX_BRUSH( RelativePath, ... ) FSlateBoxBrush( FDataprepEditorStyle::InContent( RelativePath, ".png" ), __VA_ARGS__ )
-
-#define DEFAULT_FONT(...) FCoreStyle::GetDefaultFontStyle(__VA_ARGS__)
+#define BOX_PLUGIN_BRUSH( RelativePath, ... ) FSlateBoxBrush( FDataprepEditorStyle::InContent( RelativePath, ".png" ), __VA_ARGS__ )
 
 TSharedPtr<FSlateStyleSet> FDataprepEditorStyle::StyleSet;
 
@@ -70,7 +69,7 @@ void FDataprepEditorStyle::Initialize()
 
 	StyleSet->Set("DataprepEditor.TrackNode.Slot", new IMAGE_PLUGIN_BRUSH("CircleBox", Icon22x22));
 
-	StyleSet->Set("DataprepEditor.Node.Body", new BOX_BRUSH("Node_Body", FMargin(16.f/64.f, 25.f/64.f, 16.f/64.f, 16.f/64.f)));
+	StyleSet->Set("DataprepEditor.Node.Body", new BOX_PLUGIN_BRUSH("Node_Body", FMargin(16.f/64.f, 25.f/64.f, 16.f/64.f, 16.f/64.f)));
 
 	StyleSet->Set( "DataprepEditor.SectionFont", DEFAULT_FONT( "Bold", 10 ) );
 
@@ -187,6 +186,10 @@ void FDataprepEditorStyle::Initialize()
 	FSlateStyleRegistry::RegisterSlateStyle(*StyleSet.Get());
 }
 
+#undef IMAGE_PLUGIN_BRUSH
+#undef BOX_PLUGIN_BRUSH
+#undef RootToContentDir
+
 void FDataprepEditorStyle::Shutdown()
 {
 	if (StyleSet.IsValid())
diff --git a/Engine/Plugins/Experimental/ActorPalette/Source/ActorPalette/Private/ActorPaletteStyle.cpp b/Engine/Plugins/Experimental/ActorPalette/Source/ActorPalette/Private/ActorPaletteStyle.cpp
--- a/Engine/Plugins/Experimental/ActorPalette/Source/ActorPalette/Private/ActorPaletteStyle.cpp
+++ b/Engine/Plugins/Experimental/ActorPalette/Source/ActorPalette/Private/ActorPaletteStyle.cpp
@@ -9,6 +9,9 @@
 #include "Rendering/SlateRenderer.h"
 #include "Styling/CoreStyle.h"
 #include "Styling/SlateStyle.h"
+#include "Styling/SlateStyleMacros.h"
+
+#define RootToContentDir Style->RootToContentDir
 
 TSharedPtr< FSlateStyleSet > FActorPaletteStyle::StyleInstance = NULL;
 
@@ -34,12 +37,6 @@ FName FActorPaletteStyle::GetStyleSetName()
 	return StyleSetName;
 }
 
-#define IMAGE_BRUSH( RelativePath, ... ) FSlateImageBrush( Style->RootToContentDir( RelativePath, TEXT(".png") ), __VA_ARGS__ )
-#define BOX_BRUSH( RelativePath, ... ) FSlateBoxBrush( Style->RootToContentDir( RelativePath, TEXT(".png") ), __VA_ARGS__ )
-#define BORDER_BRUSH( RelativePath, ... ) FSlateBorderBrush( Style->RootToContentDir( RelativePath, TEXT(".png") ), __VA_ARGS__ )
-#define TTF_FONT( RelativePath, ... ) FSlateFontInfo( Style->RootToContentDir( RelativePath, TEXT(".ttf") ), __VA_ARGS__ )
-#define OTF_FONT( RelativePath, ... ) FSlateFontInfo( Style->RootToContentDir( RelativePath, TEXT(".otf") ), __VA_ARGS__ )
-
 const FVector2D Icon16x16(16.0f, 16.0f);
 const FVector2D Icon20x20(20.0f, 20.0f);
 const FVector2D Icon40x40(40.0f, 40.0f);
@@ -67,11 +64,7 @@ TSharedRef< FSlateStyleSet > FActorPaletteStyle::Create()
 	return Style;
 }
 
-#undef IMAGE_BRUSH
-#undef BOX_BRUSH
-#undef BORDER_BRUSH
-#undef TTF_FONT
-#undef OTF_FONT
+#undef RootToContentDir
 
 void FActorPaletteStyle::ReloadTextures()
 {
diff --git a/Engine/Plugins/Experimental/Animation/ContextualAnimation/Source/ContextualAnimationEditor/Private/ContextualAnimEditorStyle.cpp b/Engine/Plugins/Experimental/Animation/ContextualAnimation/Source/ContextualAnimationEditor/Private/ContextualAnimEditorStyle.cpp
--- a/Engine/Plugins/Experimental/Animation/ContextualAnimation/Source/ContextualAnimationEditor/Private/ContextualAnimEditorStyle.cpp
+++ b/Engine/Plugins/Experimental/Animation/ContextualAnimation/Source/ContextualAnimationEditor/Private/ContextualAnimEditorStyle.cpp
@@ -9,6 +9,9 @@
 #include "Misc/Paths.h"
 #include "Rendering/SlateRenderer.h"
 #include "Styling/SlateStyle.h"
+#include "Styling/SlateStyleMacros.h"
+
+#define RootToContentDir Style->RootToContentDir
 
 TSharedPtr< FSlateStyleSet > FContextualAnimEditorStyle::StyleInstance = NULL;
 
@@ -34,12 +37,6 @@ FName FContextualAnimEditorStyle::GetStyleSetName()
 	return StyleSetName;
 }
 
-#define IMAGE_BRUSH( RelativePath, ... ) FSlateImageBrush( Style->RootToContentDir( RelativePath, TEXT(".png") ), __VA_ARGS__ )
-#define BOX_BRUSH( RelativePath, ... ) FSlateBoxBrush( Style->RootToContentDir( RelativePath, TEXT(".png") ), __VA_ARGS__ )
-#define BORDER_BRUSH( RelativePath, ... ) FSlateBorderBrush( Style->RootToContentDir( RelativePath, TEXT(".png") ), __VA_ARGS__ )
-#define TTF_FONT( RelativePath, ... ) FSlateFontInfo( Style->RootToContentDir( RelativePath, TEXT(".ttf") ), __VA_ARGS__ )
-#define OTF_FONT( RelativePath, ... ) FSlateFontInfo( Style->RootToContentDir( RelativePath, TEXT(".otf") ), __VA_ARGS__ )
-
 const FVector2D Icon16x16(16.0f, 16.0f);
 const FVector2D Icon20x20(20.0f, 20.0f);
 const FVector2D Icon40x40(40.0f, 40.0f);
@@ -58,11 +55,7 @@ TSharedRef< FSlateStyleSet > FContextualAnimEditorStyle::Create()
 	return Style;
 }
 
-#undef IMAGE_BRUSH
-#undef BOX_BRUSH
-#undef BORDER_BRUSH
-#undef TTF_FONT
-#undef OTF_FONT
+#undef RootToContentDir
 
 void FContextualAnimEditorStyle::ReloadTextures()
 {
diff --git a/Engine/Plugins/Experimental/Avalanche/Source/AvalancheOutliner/Private/Columns/Slate/SAvaOutlinerLabelComponent.cpp b/Engine/Plugins/Experimental/Avalanche/Source/AvalancheOutliner/Private/Columns/Slate/SAvaOutlinerLabelComponent.cpp
--- a/Engine/Plugins/Experimental/Avalanche/Source/AvalancheOutliner/Private/Columns/Slate/SAvaOutlinerLabelComponent.cpp
+++ b/Engine/Plugins/Experimental/Avalanche/Source/AvalancheOutliner/Private/Columns/Slate/SAvaOutlinerLabelComponent.cpp
@@ -4,6 +4,7 @@
 #include "Fonts/SlateFontInfo.h"
 #include "Item/AvaOutlinerComponent.h"
 #include "Misc/Paths.h"
+#include "Styling/SlateStyleMacros.h"
 
 #define LOCTEXT_NAMESPACE "SAvaOutlinerLabelComponent"
 
@@ -20,7 +21,7 @@ const FInlineEditableTextBlockStyle* SAvaOutlinerLabelComponent::GetTextBlockSty
 	static bool bFontSet = false;
 	if (!bFontSet)
 	{
-		Style.TextStyle.Font = FSlateFontInfo(FPaths::EngineContentDir() / TEXT("Slate/Fonts/Roboto-Italic.ttf"), 10);
+		Style.TextStyle.Font = DEFAULT_FONT("Italic", 10);
 		bFontSet = true;
 	}
 	return &Style;
diff --git a/Engine/Plugins/Experimental/ChaosFlesh/Source/ChaosFleshEditor/Private/ChaosFlesh/ChaosFleshEditorPlugin.cpp b/Engine/Plugins/Experimental/ChaosFlesh/Source/ChaosFleshEditor/Private/ChaosFlesh/ChaosFleshEditorPlugin.cpp
--- a/Engine/Plugins/Experimental/ChaosFlesh/Source/ChaosFleshEditor/Private/ChaosFlesh/ChaosFleshEditorPlugin.cpp
+++ b/Engine/Plugins/Experimental/ChaosFlesh/Source/ChaosFleshEditor/Private/ChaosFlesh/ChaosFleshEditorPlugin.cpp
@@ -17,9 +17,6 @@
 
 #define LOCTEXT_NAMESPACE "FleshEditor"
 
-#define BOX_BRUSH(StyleSet, RelativePath, ...) FSlateBoxBrush(StyleSet->RootToContentDir(RelativePath, TEXT(".png")), __VA_ARGS__)
-#define DEFAULT_FONT(...) FCoreStyle::GetDefaultFontStyle(__VA_ARGS__)
-
 void IChaosFleshEditorPlugin::StartupModule()
 {
 	FChaosFleshEditorStyle::Get();
@@ -92,5 +89,4 @@ void IChaosFleshEditorPlugin::ShutdownModule()
 
 IMPLEMENT_MODULE(IChaosFleshEditorPlugin, ChaosFleshEditor)
 
-
 #undef LOCTEXT_NAMESPACE
diff --git a/Engine/Plugins/Experimental/CodeEditor/Source/CodeEditor/Private/CodeEditorStyle.cpp b/Engine/Plugins/Experimental/CodeEditor/Source/CodeEditor/Private/CodeEditorStyle.cpp
--- a/Engine/Plugins/Experimental/CodeEditor/Source/CodeEditor/Private/CodeEditorStyle.cpp
+++ b/Engine/Plugins/Experimental/CodeEditor/Source/CodeEditor/Private/CodeEditorStyle.cpp
@@ -9,13 +9,11 @@
 #include "Styling/CoreStyle.h"
 #include "Misc/Paths.h"
 #include "Styling/SlateStyle.h"
+#include "Styling/SlateStyleMacros.h"
 
-TSharedPtr< FSlateStyleSet > FCodeEditorStyle::StyleSet = nullptr;
+#define RootToContentDir StyleSet->RootToContentDir
 
-#define IMAGE_BRUSH( RelativePath, ... ) FSlateImageBrush( StyleSet->RootToContentDir( RelativePath, TEXT(".png") ), __VA_ARGS__ )
-#define BOX_BRUSH( RelativePath, ... ) FSlateBoxBrush( StyleSet->RootToContentDir( RelativePath, TEXT(".png") ), __VA_ARGS__ )
-#define BORDER_BRUSH( RelativePath, ... ) FSlateBorderBrush( StyleSet->RootToContentDir( RelativePath, TEXT(".png") ), __VA_ARGS__ )
-#define DEFAULT_FONT(...) FCoreStyle::GetDefaultFontStyle(__VA_ARGS__)
+TSharedPtr< FSlateStyleSet > FCodeEditorStyle::StyleSet = nullptr;
 
 // Const icon sizes
 static const FVector2D Icon8x8(8.0f, 8.0f);
@@ -105,10 +103,7 @@ void FCodeEditorStyle::Initialize()
 	FSlateStyleRegistry::RegisterSlateStyle( *StyleSet.Get() );
 }
 
-#undef IMAGE_BRUSH
-#undef BOX_BRUSH
-#undef BORDER_BRUSH
-#undef DEFAULT_FONT
+#undef RootToContentDir
 
 void FCodeEditorStyle::Shutdown()
 {
diff --git a/Engine/Plugins/Experimental/HairModelingToolset/Source/HairModelingToolset/Private/HairModelingToolsStyle.cpp b/Engine/Plugins/Experimental/HairModelingToolset/Source/HairModelingToolset/Private/HairModelingToolsStyle.cpp
--- a/Engine/Plugins/Experimental/HairModelingToolset/Source/HairModelingToolset/Private/HairModelingToolsStyle.cpp
+++ b/Engine/Plugins/Experimental/HairModelingToolset/Source/HairModelingToolset/Private/HairModelingToolsStyle.cpp
@@ -10,8 +10,6 @@
 #include "Styling/SlateStyleMacros.h"
 
 #define IMAGE_PLUGIN_BRUSH( RelativePath, ... ) FSlateImageBrush( FHairModelingToolsStyle::InContent( RelativePath, ".png" ), __VA_ARGS__ )
-
-// This is to fix the issue that SlateStyleMacros like IMAGE_BRUSH look for RootToContentDir but StyleSet->RootToContentDir is how this style is set up
 #define RootToContentDir StyleSet->RootToContentDir
 
 FString FHairModelingToolsStyle::InContent(const FString& RelativePath, const ANSICHAR* Extension)
@@ -65,9 +63,7 @@ void FHairModelingToolsStyle::Initialize()
 END_SLATE_FUNCTION_BUILD_OPTIMIZATION
 
 #undef IMAGE_PLUGIN_BRUSH
-#undef IMAGE_BRUSH
-#undef BOX_BRUSH
-#undef DEFAULT_FONT
+#undef RootToContentDir
 
 void FHairModelingToolsStyle::Shutdown()
 {
diff --git a/Engine/Plugins/Experimental/PluginReferenceViewer/Source/PluginReferenceViewer/Private/PluginReferenceViewerStyle.cpp b/Engine/Plugins/Experimental/PluginReferenceViewer/Source/PluginReferenceViewer/Private/PluginReferenceViewerStyle.cpp
--- a/Engine/Plugins/Experimental/PluginReferenceViewer/Source/PluginReferenceViewer/Private/PluginReferenceViewerStyle.cpp
+++ b/Engine/Plugins/Experimental/PluginReferenceViewer/Source/PluginReferenceViewer/Private/PluginReferenceViewerStyle.cpp
@@ -9,13 +9,12 @@
 #include "Misc/Paths.h"
 #include "Styling/StyleColors.h"
 #include "Styling/CoreStyle.h"
+#include "Styling/SlateStyleMacros.h"
 #include "Styling/SlateStyleRegistry.h"
 
 
 FName FPluginReferenceViewerStyle::StyleName("PluginReferenceViewerStyle");
 
-#define FONT(...) FSlateFontInfo(FCoreStyle::GetDefaultFont(), __VA_ARGS__)
-
 FPluginReferenceViewerStyle::FPluginReferenceViewerStyle()
 	: FSlateStyleSet(StyleName)
 {
@@ -28,7 +27,7 @@ FPluginReferenceViewerStyle::FPluginReferenceViewerStyle()
 	FTextBlockStyle NormalText = GetParentStyle()->GetWidgetStyle<FTextBlockStyle>("NormalText");
 	FEditableTextBoxStyle NormalEditableTextBoxStyle = GetParentStyle()->GetWidgetStyle<FEditableTextBoxStyle>("NormalEditableTextBox");
 
-	FSlateFontInfo TitleFont = FONT(12, "Bold");
+	FSlateFontInfo TitleFont = DEFAULT_FONT("Bold", 12);
 
 	// Text Styles 
 	FTextBlockStyle GraphNodeTitle = FTextBlockStyle(NormalText)
diff --git a/Engine/Plugins/Experimental/TextureGraph/Source/TextureGraphEditor/Private/TG_Style.cpp b/Engine/Plugins/Experimental/TextureGraph/Source/TextureGraphEditor/Private/TG_Style.cpp
--- a/Engine/Plugins/Experimental/TextureGraph/Source/TextureGraphEditor/Private/TG_Style.cpp
+++ b/Engine/Plugins/Experimental/TextureGraph/Source/TextureGraphEditor/Private/TG_Style.cpp
@@ -15,8 +15,6 @@
 #include "HAL/FileManagerGeneric.h"
 #include "Styling/StyleColors.h"
 
-#define FONT(...) FSlateFontInfo(FCoreStyle::GetDefaultFont(), __VA_ARGS__)
-
 void FTG_Style::Register()
 {
 	FSlateStyleRegistry::RegisterSlateStyle(Get());
diff --git a/Engine/Plugins/Experimental/TextureGraph/Source/TextureGraphInsightEditor/Private/TextureGraphInsightEditorStyle.cpp b/Engine/Plugins/Experimental/TextureGraph/Source/TextureGraphInsightEditor/Private/TextureGraphInsightEditorStyle.cpp
--- a/Engine/Plugins/Experimental/TextureGraph/Source/TextureGraphInsightEditor/Private/TextureGraphInsightEditorStyle.cpp
+++ b/Engine/Plugins/Experimental/TextureGraph/Source/TextureGraphInsightEditor/Private/TextureGraphInsightEditorStyle.cpp
@@ -5,6 +5,9 @@
 #include "Framework/Application/SlateApplication.h"
 #include "Slate/SlateGameResources.h"
 #include "Interfaces/IPluginManager.h"
+#include "Styling/SlateStyleMacros.h"
+
+#define RootToContentDir Style->RootToContentDir
 
 TSharedPtr< FSlateStyleSet > FTextureGraphInsightEditorStyle::StyleInstance = NULL;
 
@@ -30,12 +33,6 @@ FName FTextureGraphInsightEditorStyle::GetStyleSetName()
 	return StyleSetName;
 }
 
-#define IMAGE_BRUSH( RelativePath, ... ) FSlateImageBrush( Style->RootToContentDir( RelativePath, TEXT(".png") ), __VA_ARGS__ )
-#define BOX_BRUSH( RelativePath, ... ) FSlateBoxBrush( Style->RootToContentDir( RelativePath, TEXT(".png") ), __VA_ARGS__ )
-#define BORDER_BRUSH( RelativePath, ... ) FSlateBorderBrush( Style->RootToContentDir( RelativePath, TEXT(".png") ), __VA_ARGS__ )
-#define TTF_FONT( RelativePath, ... ) FSlateFontInfo( Style->RootToContentDir( RelativePath, TEXT(".ttf") ), __VA_ARGS__ )
-#define OTF_FONT( RelativePath, ... ) FSlateFontInfo( Style->RootToContentDir( RelativePath, TEXT(".otf") ), __VA_ARGS__ )
-
 const FVector2D Icon16x16(16.0f, 16.0f);
 const FVector2D Icon20x20(20.0f, 20.0f);
 const FVector2D Icon40x40(40.0f, 40.0f);
@@ -50,12 +47,6 @@ TSharedRef< FSlateStyleSet > FTextureGraphInsightEditorStyle::Create()
 	return Style;
 }
 
-#undef IMAGE_BRUSH
-#undef BOX_BRUSH
-#undef BORDER_BRUSH
-#undef TTF_FONT
-#undef OTF_FONT
-
 void FTextureGraphInsightEditorStyle::ReloadTextures()
 {
 	if (FSlateApplication::IsInitialized())
@@ -68,3 +59,5 @@ const ISlateStyle& FTextureGraphInsightEditorStyle::Get()
 {
 	return *StyleInstance;
 }
+
+#undef RootToContentDir
diff --git a/Engine/Plugins/Experimental/ToolPresets/Source/ToolPresetEditor/Private/ToolPresetEditorStyle.cpp b/Engine/Plugins/Experimental/ToolPresets/Source/ToolPresetEditor/Private/ToolPresetEditorStyle.cpp
--- a/Engine/Plugins/Experimental/ToolPresets/Source/ToolPresetEditor/Private/ToolPresetEditorStyle.cpp
+++ b/Engine/Plugins/Experimental/ToolPresets/Source/ToolPresetEditor/Private/ToolPresetEditorStyle.cpp
@@ -11,8 +11,6 @@
 #include "Styling/SlateStyleRegistry.h"
 
 #define IMAGE_PLUGIN_BRUSH( RelativePath, ... ) FSlateImageBrush( FToolPresetEditorStyle::InContent( RelativePath, ".png" ), __VA_ARGS__ )
-
-// This is to fix the issue that SlateStyleMacros like IMAGE_BRUSH look for RootToContentDir but StyleSet->RootToContentDir is how this style is set up
 #define RootToContentDir StyleSet->RootToContentDir
 
 FString FToolPresetEditorStyle::InContent(const FString& RelativePath, const ANSICHAR* Extension)
@@ -75,9 +73,7 @@ void FToolPresetEditorStyle::Initialize()
 END_SLATE_FUNCTION_BUILD_OPTIMIZATION
 
 #undef IMAGE_PLUGIN_BRUSH
-#undef IMAGE_BRUSH
-#undef BOX_BRUSH
-#undef DEFAULT_FONT
+#undef RootToContentDir
 
 void FToolPresetEditorStyle::Shutdown()
 {
diff --git a/Engine/Plugins/Experimental/UserToolBoxCore/Source/UserToolBoxCore/Private/UserToolBoxStyle.cpp b/Engine/Plugins/Experimental/UserToolBoxCore/Source/UserToolBoxCore/Private/UserToolBoxStyle.cpp
--- a/Engine/Plugins/Experimental/UserToolBoxCore/Source/UserToolBoxCore/Private/UserToolBoxStyle.cpp
+++ b/Engine/Plugins/Experimental/UserToolBoxCore/Source/UserToolBoxCore/Private/UserToolBoxStyle.cpp
@@ -7,6 +7,9 @@
 #include "Framework/Application/SlateApplication.h"
 #include "Interfaces/IPluginManager.h"
 #include "Styling/StyleColors.h"
+#include "Styling/SlateStyleMacros.h"
+
+#define RootToContentDir Style->RootToContentDir
 
 TSharedPtr< FSlateStyleSet > FUserToolBoxStyle::StyleInstance = nullptr;
 TArray<FString> FUserToolBoxStyle::ExternalBrushIds;
@@ -32,12 +35,6 @@ FName FUserToolBoxStyle::GetStyleSetName()
 	return StyleSetName;
 }
 
-#define IMAGE_BRUSH( RelativePath, ... ) FSlateImageBrush( Style->RootToContentDir( RelativePath, TEXT(".png") ), __VA_ARGS__ )
-#define BOX_BRUSH( RelativePath, ... ) FSlateBoxBrush( Style->RootToContentDir( RelativePath, TEXT(".png") ), __VA_ARGS__ )
-#define BORDER_BRUSH( RelativePath, ... ) FSlateBorderBrush( Style->RootToContentDir( RelativePath, TEXT(".png") ), __VA_ARGS__ )
-#define TTF_FONT( RelativePath, ... ) FSlateFontInfo( Style->RootToContentDir( RelativePath, TEXT(".ttf") ), __VA_ARGS__ )
-#define OTF_FONT( RelativePath, ... ) FSlateFontInfo( Style->RootToContentDir( RelativePath, TEXT(".otf") ), __VA_ARGS__ )
-
 const FVector2D Icon16x16(16.0f, 16.0f);
 const FVector2D Icon15x15(16.0f, 16.0f);
 const FVector2D Icon30x30(16.0f, 16.0f);
@@ -59,12 +56,6 @@ TSharedRef< FSlateStyleSet > FUserToolBoxStyle::Create()
 	return Style;
 }
 
-#undef IMAGE_BRUSH
-#undef BOX_BRUSH
-#undef BORDER_BRUSH
-#undef TTF_FONT
-#undef OTF_FONT
-
 void FUserToolBoxStyle::AddExternalImageBrushes(const TArray<FIconInfo>& IconInfos)
 {
 	if (StyleInstance==nullptr)
@@ -101,3 +92,5 @@ const ISlateStyle& FUserToolBoxStyle::Get()
 {
 	return *StyleInstance;
 }
+
+#undef RootToContentDir
diff --git a/Engine/Plugins/Experimental/Water/Source/Editor/Private/WaterUIStyle.cpp b/Engine/Plugins/Experimental/Water/Source/Editor/Private/WaterUIStyle.cpp
--- a/Engine/Plugins/Experimental/Water/Source/Editor/Private/WaterUIStyle.cpp
+++ b/Engine/Plugins/Experimental/Water/Source/Editor/Private/WaterUIStyle.cpp
@@ -6,13 +6,9 @@
 #include "Interfaces/IPluginManager.h"
 #include "Misc/Paths.h"
 #include "Styling/SlateStyle.h"
+#include "Styling/SlateStyleMacros.h"
 
-
-#define IMAGE_BRUSH(RelativePath, ...) FSlateImageBrush(StyleSet->RootToContentDir(RelativePath, TEXT(".png")), __VA_ARGS__)
-#define IMAGE_BRUSH_SVG(RelativePath, ...) FSlateVectorImageBrush(StyleSet->RootToContentDir(RelativePath, TEXT(".svg")), __VA_ARGS__)
-#define BOX_BRUSH(RelativePath, ...) FSlateBoxBrush(StyleSet->RootToContentDir(RelativePath, TEXT(".png")), __VA_ARGS__)
-#define DEFAULT_FONT(...) FCoreStyle::GetDefaultFontStyle(__VA_ARGS__)
-
+#define RootToContentDir StyleSet->RootToContentDir
 
 TSharedPtr< FSlateStyleSet > FWaterUIStyle::StyleSet = nullptr;
 TSharedPtr< class ISlateStyle > FWaterUIStyle::Get() { return StyleSet; }
@@ -79,10 +75,7 @@ void FWaterUIStyle::Initialize()
 	FSlateStyleRegistry::RegisterSlateStyle(*StyleSet.Get());
 };
 
-#undef IMAGE_BRUSH_SVG
-#undef IMAGE_BRUSH
-#undef BOX_BRUSH
-#undef DEFAULT_FONT
+#undef RootToContentDir
 
 void FWaterUIStyle::Shutdown()
 {
diff --git a/Engine/Plugins/Experimental/Web/WebAPI/Source/WebAPIEditor/Private/WebAPIEditorStyle.cpp b/Engine/Plugins/Experimental/Web/WebAPI/Source/WebAPIEditor/Private/WebAPIEditorStyle.cpp
--- a/Engine/Plugins/Experimental/Web/WebAPI/Source/WebAPIEditor/Private/WebAPIEditorStyle.cpp
+++ b/Engine/Plugins/Experimental/Web/WebAPI/Source/WebAPIEditor/Private/WebAPIEditorStyle.cpp
@@ -10,6 +10,10 @@
 #include "Styling/SlateStyleRegistry.h"
 #include "Styling/StyleColors.h"
 
+// This is to fix the issue that SlateStyleMacros like IMAGE_BRUSH look for RootToContentDir but StyleSet->RootToContentDir is how this style is set up
+#define RootToContentDir UE::WebAPI::Private::StyleInstance->RootToContentDir
+#define RootToCoreContentDir UE::WebAPI::Private::StyleInstance->RootToCoreContentDir
+
 namespace UE
 {
 	namespace WebAPI
@@ -30,10 +34,6 @@ namespace UE
 
 TUniquePtr<FSlateStyleSet> FWebAPIEditorStyle::StyleInstance = nullptr;
 
-// This is to fix the issue that SlateStyleMacros like IMAGE_BRUSH look for RootToContentDir but StyleSet->RootToContentDir is how this style is set up
-#define RootToContentDir UE::WebAPI::Private::StyleInstance->RootToContentDir
-#define RootToCoreContentDir UE::WebAPI::Private::StyleInstance->RootToCoreContentDir
-
 void FWebAPIEditorStyle::Register()
 {
 	UE::WebAPI::Private::StyleInstance = MakeUnique<FSlateStyleSet>(UE::WebAPI::Private::NAME_StyleName);
@@ -62,9 +62,5 @@ const ISlateStyle& FWebAPIEditorStyle::Get()
 	return *UE::WebAPI::Private::StyleInstance.Get();
 }
 
-#undef IMAGE_BRUSH
-#undef BOX_BRUSH
-#undef BORDER_BRUSH
-#undef DEFAULT_FONT
-#undef CORE_BOX_BRUSH
-#undef CORE_IMAGE_BRUSH
+#undef RootToContentDir
+#undef RootToCoreContentDir
diff --git a/Engine/Plugins/Experimental/WidgetEditorToolPalette/Source/WidgetEditorToolPalette/Private/WidgetEditorToolPaletteStyle.cpp b/Engine/Plugins/Experimental/WidgetEditorToolPalette/Source/WidgetEditorToolPalette/Private/WidgetEditorToolPaletteStyle.cpp
--- a/Engine/Plugins/Experimental/WidgetEditorToolPalette/Source/WidgetEditorToolPalette/Private/WidgetEditorToolPaletteStyle.cpp
+++ b/Engine/Plugins/Experimental/WidgetEditorToolPalette/Source/WidgetEditorToolPalette/Private/WidgetEditorToolPaletteStyle.cpp
@@ -9,10 +9,7 @@
 #include "Styling/SlateStyleMacros.h"
 #include "Styling/AppStyle.h"
 
-
 #define IMAGE_PLUGIN_BRUSH( RelativePath, ... ) FSlateImageBrush( FWidgetEditorToolPaletteStyle::InContent( RelativePath, ".png" ), __VA_ARGS__ )
-
-// This is to fix the issue that SlateStyleMacros like IMAGE_BRUSH look for RootToContentDir but StyleSet->RootToContentDir is how this style is set up
 #define RootToContentDir StyleSet->RootToContentDir
 
 FString FWidgetEditorToolPaletteStyle::InContent(const FString& RelativePath, const ANSICHAR* Extension)
@@ -87,9 +84,7 @@ void FWidgetEditorToolPaletteStyle::Initialize()
 END_SLATE_FUNCTION_BUILD_OPTIMIZATION
 
 #undef IMAGE_PLUGIN_BRUSH
-#undef IMAGE_BRUSH
-#undef BOX_BRUSH
-#undef DEFAULT_FONT
+#undef RootToContentDir
 
 void FWidgetEditorToolPaletteStyle::Shutdown()
 {
diff --git a/Engine/Plugins/FX/Niagara/Source/NiagaraEditorWidgets/Private/Stack/SNiagaraDeterminismToggle.cpp b/Engine/Plugins/FX/Niagara/Source/NiagaraEditorWidgets/Private/Stack/SNiagaraDeterminismToggle.cpp
--- a/Engine/Plugins/FX/Niagara/Source/NiagaraEditorWidgets/Private/Stack/SNiagaraDeterminismToggle.cpp
+++ b/Engine/Plugins/FX/Niagara/Source/NiagaraEditorWidgets/Private/Stack/SNiagaraDeterminismToggle.cpp
@@ -36,7 +36,7 @@ void SNiagaraDeterminismToggle::Construct(const FArguments& InArgs, TSharedPtr<F
 		PropertyDescription = MatchingProperty->GetToolTipText();
 	}
 
-	FSlateFontInfo FontInfo = FSlateFontInfo(FPaths::EngineContentDir() / TEXT("Slate/Fonts/Roboto-Regular.ttf"), 9);
+	FSlateFontInfo FontInfo = FSlateFontInfo(FCoreStyle::GetDefaultFont(), 9, "Regular");
 	FontInfo.OutlineSettings = FFontOutlineSettings(1, FLinearColor::White);
 	
 	ChildSlot
diff --git a/Engine/Plugins/Importers/USDImporter/Source/USDStageEditor/Private/SUSDStageEditorStyle.cpp b/Engine/Plugins/Importers/USDImporter/Source/USDStageEditor/Private/SUSDStageEditorStyle.cpp
--- a/Engine/Plugins/Importers/USDImporter/Source/USDStageEditor/Private/SUSDStageEditorStyle.cpp
+++ b/Engine/Plugins/Importers/USDImporter/Source/USDStageEditor/Private/SUSDStageEditorStyle.cpp
@@ -6,12 +6,10 @@
 #include "SlateOptMacros.h"
 #include "Styling/SlateStyleRegistry.h"
 #include "Styling/SlateTypes.h"
+#include "Styling/SlateStyleMacros.h"
 
 #define IMAGE_PLUGIN_BRUSH(RelativePath, ...) FSlateImageBrush(FUsdStageEditorStyle::InContent(RelativePath, ".png"), __VA_ARGS__)
-#define IMAGE_BRUSH(RelativePath, ...) FSlateImageBrush(StyleSet->RootToContentDir(RelativePath, TEXT(".png")), __VA_ARGS__)
-#define IMAGE_BRUSH_SVG(RelativePath, ...) FSlateVectorImageBrush(StyleSet->RootToContentDir(RelativePath, TEXT(".svg")), __VA_ARGS__)
-#define BOX_BRUSH(RelativePath, ...) FSlateBoxBrush(StyleSet->RootToContentDir(RelativePath, TEXT(".png")), __VA_ARGS__)
-#define DEFAULT_FONT(...) FCoreStyle::GetDefaultFontStyle(__VA_ARGS__)
+#define RootToContentDir StyleSet->RootToContentDir
 
 FString FUsdStageEditorStyle::InContent(const FString& RelativePath, const ANSICHAR* Extension)
 {
@@ -90,9 +88,7 @@ void FUsdStageEditorStyle::Initialize()
 END_SLATE_FUNCTION_BUILD_OPTIMIZATION
 
 #undef IMAGE_PLUGIN_BRUSH
-#undef IMAGE_BRUSH
-#undef BOX_BRUSH
-#undef DEFAULT_FONT
+#undef RootToContentDir
 
 void FUsdStageEditorStyle::Shutdown()
 {
diff --git a/Engine/Plugins/Media/BinkMedia/Source/BinkMediaPlayerEditor/Private/BinkMediaPlayerEditorModule.cpp b/Engine/Plugins/Media/BinkMedia/Source/BinkMediaPlayerEditor/Private/BinkMediaPlayerEditorModule.cpp
--- a/Engine/Plugins/Media/BinkMedia/Source/BinkMediaPlayerEditor/Private/BinkMediaPlayerEditorModule.cpp
+++ b/Engine/Plugins/Media/BinkMedia/Source/BinkMediaPlayerEditor/Private/BinkMediaPlayerEditorModule.cpp
@@ -20,15 +20,10 @@
 #include "Misc/Paths.h"
 #include "ThumbnailRendering/TextureThumbnailRenderer.h"
 #include "Styling/SlateStyle.h"
+#include "Styling/SlateStyleMacros.h"
 
 #define LOCTEXT_NAMESPACE "FBinkMediaPlayerEditorModule"
 
-#define IMAGE_BRUSH( RelativePath, ... ) FSlateImageBrush(RootToContentDir(RelativePath, TEXT(".png")), __VA_ARGS__)
-#define BOX_BRUSH( RelativePath, ... ) FSlateBoxBrush(RootToContentDir(RelativePath, TEXT(".png")), __VA_ARGS__)
-#define BORDER_BRUSH( RelativePath, ... ) FSlateBorderBrush(RootToContentDir(RelativePath, TEXT(".png")), __VA_ARGS__)
-#define TTF_FONT( RelativePath, ... ) FSlateFontInfo(RootToContentDir(RelativePath, TEXT(".ttf")), __VA_ARGS__)
-#define OTF_FONT( RelativePath, ... ) FSlateFontInfo(RootToContentDir(RelativePath, TEXT(".otf")), __VA_ARGS__)
-
 struct FBinkMediaPlayerEditorStyle final : FSlateStyleSet
 {
 	FBinkMediaPlayerEditorStyle() : FSlateStyleSet("BinkMediaPlayerEditorStyle")
@@ -56,12 +51,6 @@ struct FBinkMediaPlayerEditorStyle final : FSlateStyleSet
 	}
 };
 
-#undef IMAGE_BRUSH
-#undef BOX_BRUSH
-#undef BORDER_BRUSH
-#undef TTF_FONT
-#undef OTF_FONT
-
 struct FBinkMediaTextureActions : FAssetTypeActions_Base 
 {
 	virtual void GetActions( const TArray<UObject*>& InObjects, FMenuBuilder& MenuBuilder ) override 
diff --git a/Engine/Plugins/Media/BinkMedia/Source/BinkMediaPlayerEditor/Private/SBinkWidgets.cpp b/Engine/Plugins/Media/BinkMedia/Source/BinkMediaPlayerEditor/Private/SBinkWidgets.cpp
--- a/Engine/Plugins/Media/BinkMedia/Source/BinkMediaPlayerEditor/Private/SBinkWidgets.cpp
+++ b/Engine/Plugins/Media/BinkMedia/Source/BinkMediaPlayerEditor/Private/SBinkWidgets.cpp
@@ -167,7 +167,7 @@ void SBinkMediaPlayerEditorViewer::Construct( const FArguments& InArgs, UBinkMed
 					[
 						// playback state
 						SNew(STextBlock)
-							.Font(FSlateFontInfo(FPaths::EngineContentDir() / TEXT("Slate/Fonts/Roboto-Regular.ttf"), 20))
+							.Font(FSlateFontInfo(FCoreStyle::GetDefaultFont(), 20, "Regular"))
 							.ShadowOffset(FVector2D(1.f, 1.f))
 							.Text(this, &SBinkMediaPlayerEditorViewer::HandleOverlayStateText)
 					]
diff --git a/Engine/Plugins/Media/MediaPlate/Source/MediaPlateEditor/Private/MediaPlateEditorStyle.cpp b/Engine/Plugins/Media/MediaPlate/Source/MediaPlateEditor/Private/MediaPlateEditorStyle.cpp
--- a/Engine/Plugins/Media/MediaPlate/Source/MediaPlateEditor/Private/MediaPlateEditorStyle.cpp
+++ b/Engine/Plugins/Media/MediaPlate/Source/MediaPlateEditor/Private/MediaPlateEditorStyle.cpp
@@ -9,16 +9,11 @@
 #include "Interfaces/IPluginManager.h"
 #include "Math/Vector2D.h"
 #include "Misc/Paths.h"
+#include "Styling/SlateStyleMacros.h"
 #include "Styling/SlateStyleRegistry.h"
 #include "Styling/SlateTypes.h"
 #include "Styling/CoreStyle.h"
 
-#define IMAGE_BRUSH(RelativePath, ...) FSlateImageBrush(RootToContentDir(RelativePath, TEXT(".png")), __VA_ARGS__)
-#define IMAGE_BRUSH_SVG(RelativePath, ...) FSlateVectorImageBrush(RootToContentDir(RelativePath, TEXT(".svg" )), __VA_ARGS__)
-#define BOX_BRUSH(RelativePath, ...) FSlateBoxBrush(RootToContentDir(RelativePath, TEXT(".png")), __VA_ARGS__)
-#define BORDER_BRUSH(RelativePath, ...) FSlateBorderBrush(RootToContentDir(RelativePath, TEXT(".png")), __VA_ARGS__)
-#define DEFAULT_FONT(...) FCoreStyle::GetDefaultFontStyle(__VA_ARGS__)
-
 /* FMediaPlateEditorStyle structors
  *****************************************************************************/
 
@@ -102,8 +97,3 @@ FMediaPlateEditorStyle::~FMediaPlateEditorStyle()
 {
 	FSlateStyleRegistry::UnRegisterSlateStyle(*this);
 }
-
-#undef IMAGE_BRUSH
-#undef BOX_BRUSH
-#undef BORDER_BRUSH
-#undef DEFAULT_FONT
diff --git a/Engine/Plugins/Media/MediaPlayerEditor/Source/MediaPlayerEditor/Private/Shared/MediaPlayerEditorStyle.cpp b/Engine/Plugins/Media/MediaPlayerEditor/Source/MediaPlayerEditor/Private/Shared/MediaPlayerEditorStyle.cpp
--- a/Engine/Plugins/Media/MediaPlayerEditor/Source/MediaPlayerEditor/Private/Shared/MediaPlayerEditorStyle.cpp
+++ b/Engine/Plugins/Media/MediaPlayerEditor/Source/MediaPlayerEditor/Private/Shared/MediaPlayerEditorStyle.cpp
@@ -87,9 +87,3 @@ FMediaPlayerEditorStyle::~FMediaPlayerEditorStyle()
 {
 	FSlateStyleRegistry::UnRegisterSlateStyle(*this);
 }
-
-
-#undef IMAGE_BRUSH
-#undef BOX_BRUSH
-#undef BORDER_BRUSH
-#undef DEFAULT_FONT
diff --git a/Engine/Plugins/Media/PixelStreaming/Source/PixelStreamingEditor/Private/PixelStreamingStyle.cpp b/Engine/Plugins/Media/PixelStreaming/Source/PixelStreamingEditor/Private/PixelStreamingStyle.cpp
--- a/Engine/Plugins/Media/PixelStreaming/Source/PixelStreamingEditor/Private/PixelStreamingStyle.cpp
+++ b/Engine/Plugins/Media/PixelStreaming/Source/PixelStreamingEditor/Private/PixelStreamingStyle.cpp
@@ -6,6 +6,9 @@
 #include "Slate/SlateGameResources.h"
 #include "Interfaces/IPluginManager.h"
 #include "Styling/SlateStyle.h"
+#include "Styling/SlateStyleMacros.h"
+
+#define RootToContentDir Style->RootToContentDir
 
 namespace UE::EditorPixelStreaming
 {
@@ -32,13 +35,6 @@ namespace UE::EditorPixelStreaming
         return StyleSetName;
     }
 
-#define IMAGE_BRUSH( RelativePath, ... ) FSlateImageBrush( Style->RootToContentDir( RelativePath, TEXT(".png") ), __VA_ARGS__ )
-#define IMAGE_BRUSH_SVG( RelativePath, ... ) FSlateVectorImageBrush(Style->RootToContentDir(RelativePath, TEXT(".svg")), __VA_ARGS__)
-#define BOX_BRUSH( RelativePath, ... ) FSlateBoxBrush( Style->RootToContentDir( RelativePath, TEXT(".png") ), __VA_ARGS__ )
-#define BORDER_BRUSH( RelativePath, ... ) FSlateBorderBrush( Style->RootToContentDir( RelativePath, TEXT(".png") ), __VA_ARGS__ )
-#define TTF_FONT( RelativePath, ... ) FSlateFontInfo( Style->RootToContentDir( RelativePath, TEXT(".ttf") ), __VA_ARGS__ )
-#define OTF_FONT( RelativePath, ... ) FSlateFontInfo( Style->RootToContentDir( RelativePath, TEXT(".otf") ), __VA_ARGS__ )
-
     const FVector2D Icon16x16(16.0f, 16.0f);
     const FVector2D Icon20x20(20.0f, 20.0f);
     const FVector2D Icon64x64(64.0f, 64.0f);
@@ -50,12 +46,6 @@ namespace UE::EditorPixelStreaming
         Style->Set("PixelStreaming.Icon", new IMAGE_BRUSH_SVG("PixelStreaming_16", Icon16x16));
         return Style;
     }
-#undef IMAGE_BRUSH
-#undef IMAGE_BRUSH_SVG
-#undef BOX_BRUSH
-#undef BORDER_BRUSH
-#undef TTF_FONT
-#undef OTF_FONT
     
     void FPixelStreamingStyle::ReloadTextures()
     {
@@ -69,4 +59,6 @@ namespace UE::EditorPixelStreaming
     {
         return *StyleInstance;
     }
-}
\ No newline at end of file
+}
+
+#undef RootToContentDir
\ No newline at end of file
diff --git a/Engine/Plugins/Media/PixelStreaming2/Source/PixelStreaming2Editor/Private/PixelStreaming2Style.cpp b/Engine/Plugins/Media/PixelStreaming2/Source/PixelStreaming2Editor/Private/PixelStreaming2Style.cpp
--- a/Engine/Plugins/Media/PixelStreaming2/Source/PixelStreaming2Editor/Private/PixelStreaming2Style.cpp
+++ b/Engine/Plugins/Media/PixelStreaming2/Source/PixelStreaming2Editor/Private/PixelStreaming2Style.cpp
@@ -6,6 +6,9 @@
 #include "Slate/SlateGameResources.h"
 #include "Interfaces/IPluginManager.h"
 #include "Styling/SlateStyle.h"
+#include "Styling/SlateStyleMacros.h"
+
+#define RootToContentDir Style->RootToContentDir
 
 namespace UE::EditorPixelStreaming2
 {
@@ -32,13 +35,6 @@ namespace UE::EditorPixelStreaming2
         return StyleSetName;
     }
 
-#define IMAGE_BRUSH( RelativePath, ... ) FSlateImageBrush( Style->RootToContentDir( RelativePath, TEXT(".png") ), __VA_ARGS__ )
-#define IMAGE_BRUSH_SVG( RelativePath, ... ) FSlateVectorImageBrush(Style->RootToContentDir(RelativePath, TEXT(".svg")), __VA_ARGS__)
-#define BOX_BRUSH( RelativePath, ... ) FSlateBoxBrush( Style->RootToContentDir( RelativePath, TEXT(".png") ), __VA_ARGS__ )
-#define BORDER_BRUSH( RelativePath, ... ) FSlateBorderBrush( Style->RootToContentDir( RelativePath, TEXT(".png") ), __VA_ARGS__ )
-#define TTF_FONT( RelativePath, ... ) FSlateFontInfo( Style->RootToContentDir( RelativePath, TEXT(".ttf") ), __VA_ARGS__ )
-#define OTF_FONT( RelativePath, ... ) FSlateFontInfo( Style->RootToContentDir( RelativePath, TEXT(".otf") ), __VA_ARGS__ )
-
     const FVector2D Icon16x16(16.0f, 16.0f);
     const FVector2D Icon20x20(20.0f, 20.0f);
     const FVector2D Icon64x64(64.0f, 64.0f);
@@ -50,12 +46,6 @@ namespace UE::EditorPixelStreaming2
         Style->Set("PixelStreaming2.Icon", new IMAGE_BRUSH_SVG("PixelStreaming2_16", Icon16x16));
         return Style;
     }
-#undef IMAGE_BRUSH
-#undef IMAGE_BRUSH_SVG
-#undef BOX_BRUSH
-#undef BORDER_BRUSH
-#undef TTF_FONT
-#undef OTF_FONT
     
     void FPixelStreaming2Style::ReloadTextures()
     {
@@ -69,4 +59,6 @@ namespace UE::EditorPixelStreaming2
     {
         return *StyleInstance;
     }
-}
\ No newline at end of file
+}
+
+#undef RootToContentDir
\ No newline at end of file
diff --git a/Engine/Plugins/Mutable/Source/CustomizableObjectEditor/Private/MuCOE/CustomizableObjectCustomSettings.cpp b/Engine/Plugins/Mutable/Source/CustomizableObjectEditor/Private/MuCOE/CustomizableObjectCustomSettings.cpp
--- a/Engine/Plugins/Mutable/Source/CustomizableObjectEditor/Private/MuCOE/CustomizableObjectCustomSettings.cpp
+++ b/Engine/Plugins/Mutable/Source/CustomizableObjectEditor/Private/MuCOE/CustomizableObjectCustomSettings.cpp
@@ -75,7 +75,7 @@ void SCustomizableObjectCustomSettings::Construct(const FArguments& InArgs)
 				.Padding(FMargin(0.f, 0.f, 4.f, 0.f))
 				[
 					SNew(STextBlock)
-					.Font(FSlateFontInfo(FPaths::EngineContentDir() / TEXT("Slate/Fonts/Roboto-Bold.ttf"), 10))
+					.Font(FSlateFontInfo(FCoreStyle::GetDefaultFont(), 10, "Bold"))
 					.Text(FText::FromString("Lighting"))
 				]
 			]
diff --git a/Engine/Plugins/Mutable/Source/CustomizableObjectEditor/Private/MuCOE/SCustomizableObjectEditorAdvancedPreviewSettings.cpp b/Engine/Plugins/Mutable/Source/CustomizableObjectEditor/Private/MuCOE/SCustomizableObjectEditorAdvancedPreviewSettings.cpp
--- a/Engine/Plugins/Mutable/Source/CustomizableObjectEditor/Private/MuCOE/SCustomizableObjectEditorAdvancedPreviewSettings.cpp
+++ b/Engine/Plugins/Mutable/Source/CustomizableObjectEditor/Private/MuCOE/SCustomizableObjectEditorAdvancedPreviewSettings.cpp
@@ -444,7 +444,7 @@ void SCustomizableObjectEditorAdvancedPreviewSettings::Construct(const FArgument
 					+ SUniformGridPanel::Slot(0, 0)
 					[
 						SNew(STextBlock)
-						.Font(FSlateFontInfo(FPaths::EngineContentDir() / TEXT("Slate/Fonts/Roboto-Bold.ttf"), 12))
+						.Font(FSlateFontInfo(FCoreStyle::GetDefaultFont(), 12, "Regular"))
 						.Text(LOCTEXT("LevelPreviewSettings", "Profile from Level Asset"))
 					]
 					+ SUniformGridPanel::Slot(1, 0)
diff --git a/Engine/Plugins/Runtime/NetworkPredictionInsights/Source/NetworkPredictionInsights/Private/NetworkPredictionInsightsStyle.h b/Engine/Plugins/Runtime/NetworkPredictionInsights/Source/NetworkPredictionInsights/Private/NetworkPredictionInsightsStyle.h
--- a/Engine/Plugins/Runtime/NetworkPredictionInsights/Source/NetworkPredictionInsights/Private/NetworkPredictionInsightsStyle.h
+++ b/Engine/Plugins/Runtime/NetworkPredictionInsights/Source/NetworkPredictionInsights/Private/NetworkPredictionInsightsStyle.h
@@ -6,12 +6,7 @@
 
 #include "Styling/SlateStyle.h"
 #include "Styling/SlateStyleRegistry.h"
-
-#define IMAGE_BRUSH(RelativePath, ...) FSlateImageBrush(RootToContentDir(RelativePath, TEXT(".png")), __VA_ARGS__)
-#define BOX_BRUSH(RelativePath, ...) FSlateBoxBrush(RootToContentDir(RelativePath, TEXT(".png")), __VA_ARGS__)
-#define BORDER_BRUSH(RelativePath, ...) FSlateBorderBrush(RootToContentDir(RelativePath, TEXT(".png")), __VA_ARGS__)
-#define TTF_FONT(RelativePath, ...) FSlateFontInfo(RootToContentDir(RelativePath, TEXT(".ttf")), __VA_ARGS__)
-#define OTF_FONT(RelativePath, ...) FSlateFontInfo(RootToContentDir(RelativePath, TEXT(".otf")), __VA_ARGS__)
+#include "Styling/SlateStyleMacros.h"
 
 class FNetworkPredictionInsightsStyle
 	: public FSlateStyleSet
@@ -85,9 +80,3 @@ public:
 
 	0.53333  0.67059  0.53333
 };
-
-#undef IMAGE_BRUSH
-#undef BOX_BRUSH
-#undef BORDER_BRUSH
-#undef TTF_FONT
-#undef OTF_FONT
diff --git a/Engine/Plugins/Runtime/ZoneGraph/Source/ZoneGraphEditor/Private/ZoneGraphEditorStyle.cpp b/Engine/Plugins/Runtime/ZoneGraph/Source/ZoneGraphEditor/Private/ZoneGraphEditorStyle.cpp
--- a/Engine/Plugins/Runtime/ZoneGraph/Source/ZoneGraphEditor/Private/ZoneGraphEditorStyle.cpp
+++ b/Engine/Plugins/Runtime/ZoneGraph/Source/ZoneGraphEditor/Private/ZoneGraphEditorStyle.cpp
@@ -1,18 +1,15 @@
 ﻿// Copyright Epic Games, Inc. All Rights Reserved.
 
 #include "ZoneGraphEditorStyle.h"
+#include "Styling/SlateStyleMacros.h"
 #include "Styling/SlateStyleRegistry.h"
 #include "Styling/CoreStyle.h"
 #include "Styling/SlateTypes.h"
 #include "Styling/AppStyle.h"
 #include "Interfaces/IPluginManager.h"
 
-#define IMAGE_BRUSH( RelativePath, ... ) FSlateImageBrush( StyleSet->RootToContentDir( RelativePath, TEXT( ".png" ) ), __VA_ARGS__ )
 #define IMAGE_PLUGIN_BRUSH( RelativePath, ... ) FSlateImageBrush( FMeshEditorStyle::InContent( RelativePath, ".png" ), __VA_ARGS__ )
-#define BOX_BRUSH( RelativePath, ... ) FSlateBoxBrush( StyleSet->RootToContentDir( RelativePath, TEXT( ".png" ) ), __VA_ARGS__ )
-#define TTF_CORE_FONT( RelativePath, ... ) FSlateFontInfo( StyleSet->RootToCoreContentDir( RelativePath, TEXT( ".ttf" ) ), __VA_ARGS__ )
-
-#define DEFAULT_FONT(...) FCoreStyle::GetDefaultFontStyle(__VA_ARGS__)
+#define RootToContentDir StyleSet->RootToContentDir
 
 TSharedPtr<FSlateStyleSet> FZoneGraphEditorStyle::StyleSet = nullptr;
 
@@ -41,6 +38,8 @@ void FZoneGraphEditorStyle::Initialize()
 	FSlateStyleRegistry::RegisterSlateStyle(*StyleSet.Get());
 }
 
+#undef IMAGE_PLUGIN_BRUSH
+#undef RootToContentDir
 
 void FZoneGraphEditorStyle::Shutdown()
 {
@@ -52,7 +51,6 @@ void FZoneGraphEditorStyle::Shutdown()
 	}
 }
 
-
 FName FZoneGraphEditorStyle::GetStyleSetName()
 {
 	static FName StyleName("ZoneGraphEditorStyle");
diff --git a/Engine/Plugins/VirtualProduction/RemoteControl/Source/RemoteControlProtocolWidgets/Private/Styling/ProtocolPanelStyle.cpp b/Engine/Plugins/VirtualProduction/RemoteControl/Source/RemoteControlProtocolWidgets/Private/Styling/ProtocolPanelStyle.cpp
--- a/Engine/Plugins/VirtualProduction/RemoteControl/Source/RemoteControlProtocolWidgets/Private/Styling/ProtocolPanelStyle.cpp
+++ b/Engine/Plugins/VirtualProduction/RemoteControl/Source/RemoteControlProtocolWidgets/Private/Styling/ProtocolPanelStyle.cpp
@@ -8,15 +8,12 @@
 #include "Styling/SlateTypes.h"
 #include "Styling/StarshipCoreStyle.h"
 #include "Styling/StyleColors.h"
+#include "Styling/SlateStyleMacros.h"
 
 #include "Interfaces/IPluginManager.h"
 
-#define IMAGE_BRUSH(RelativePath, ...) FSlateImageBrush(StyleSet->RootToContentDir(RelativePath, TEXT(".png")), __VA_ARGS__)
-#define IMAGE_BRUSH_SVG( RelativePath, ... ) FSlateVectorImageBrush(StyleSet->RootToContentDir(RelativePath, TEXT(".svg")), __VA_ARGS__)
-#define BOX_BRUSH(RelativePath, ...) FSlateBoxBrush(StyleSet->RootToContentDir(RelativePath, TEXT(".png")), __VA_ARGS__)
-#define CORE_BOX_BRUSH( RelativePath, ... ) FSlateBoxBrush(StyleSet->RootToCoreContentDir(RelativePath, TEXT(".png") ), __VA_ARGS__)
 #define BOX_PLUGIN_BRUSH( RelativePath, ... ) FSlateBoxBrush(FProtocolPanelStyle::InContent( RelativePath, ".png" ), __VA_ARGS__)
-#define DEFAULT_FONT(...) FAppStyle::Get().GetDefaultFontStyle(__VA_ARGS__)
+#define RootToContentDir StyleSet->RootToContentDir
 
 TSharedPtr<FSlateStyleSet> FProtocolPanelStyle::StyleSet;
 
@@ -97,9 +94,5 @@ void FProtocolPanelStyle::SetupWidgetStyles(TSharedRef<FSlateStyleSet> InStyle)
 	StyleSet->Set("ProtocolsPanel.Widgets.Mask", ProtocolMaskWidgetStyle);
 }
 
-#undef IMAGE_BRUSH
-#undef IMAGE_BRUSH_SVG
-#undef BOX_BRUSH
-#undef CORE_BOX_BRUSH
 #undef BOX_PLUGIN_BRUSH
-#undef DEFAULT_FONT
+#undef RootToContentDir
diff --git a/Engine/Plugins/VirtualProduction/RemoteControl/Source/RemoteControlUI/Private/UI/RemoteControlExposeMenuStyle.cpp b/Engine/Plugins/VirtualProduction/RemoteControl/Source/RemoteControlUI/Private/UI/RemoteControlExposeMenuStyle.cpp
--- a/Engine/Plugins/VirtualProduction/RemoteControl/Source/RemoteControlUI/Private/UI/RemoteControlExposeMenuStyle.cpp
+++ b/Engine/Plugins/VirtualProduction/RemoteControl/Source/RemoteControlUI/Private/UI/RemoteControlExposeMenuStyle.cpp
@@ -5,19 +5,17 @@
 #include "Styling/AppStyle.h"
 #include "Styling/CoreStyle.h"
 #include "Styling/SlateStyle.h"
+#include "Styling/SlateStyleMacros.h"
 #include "Styling/SlateStyleRegistry.h"
 #include "Styling/SlateTypes.h"
 #include "Styling/StarshipCoreStyle.h"
 #include "Styling/StyleColors.h"
+#include "Styling/SlateStyleMacros.h"
 
 #define IMAGE_PLUGIN_BRUSH( RelativePath, ... ) FSlateImageBrush(FRemoteControlExposeMenuStyle::InContent(RelativePath, ".png" ), __VA_ARGS__)
-#define IMAGE_BRUSH(RelativePath, ...) FSlateImageBrush(StyleSet->RootToContentDir(RelativePath, TEXT(".png")), __VA_ARGS__)
-#define IMAGE_BRUSH_SVG( RelativePath, ... ) FSlateVectorImageBrush(StyleSet->RootToContentDir(RelativePath, TEXT(".svg")), __VA_ARGS__)
 #define IMAGE_PLUGIN_BRUSH_SVG( RelativePath, ... ) FSlateVectorImageBrush(FRemoteControlExposeMenuStyle::InContent(RelativePath, ".svg"), __VA_ARGS__)
-#define BOX_BRUSH(RelativePath, ...) FSlateBoxBrush(StyleSet->RootToContentDir(RelativePath, TEXT(".png")), __VA_ARGS__)
-#define CORE_BOX_BRUSH( RelativePath, ... ) FSlateBoxBrush(StyleSet->RootToCoreContentDir(RelativePath, TEXT(".png") ), __VA_ARGS__)
 #define BOX_PLUGIN_BRUSH( RelativePath, ... ) FSlateBoxBrush(FRemoteControlExposeMenuStyle::InContent( RelativePath, ".png" ), __VA_ARGS__)
-#define CORE_FONT(...) FSlateFontInfo(FCoreStyle::GetDefaultFont(), __VA_ARGS__)
+#define RootToContentDir StyleSet->RootToContentDir
 
 TSharedPtr<FSlateStyleSet> FRemoteControlExposeMenuStyle::StyleSet;
 
@@ -109,7 +107,7 @@ void FRemoteControlExposeMenuStyle::Initialize()
 	StyleSet->Set("RemoteControlExposeMenu.Separator", new FSlateColorBrush(FStyleColors::White25));
 	StyleSet->Set("RemoteControlExposeMenu.Separator.Padding", FMargin(12.0f, 6.f, 12.0f, 6.f));
 
-	FSlateFontInfo XSFont(CORE_FONT(7, "Bold"));
+	FSlateFontInfo XSFont(DEFAULT_FONT("Bold", 7));
 	XSFont.LetterSpacing =  250;
 	const FTextBlockStyle SmallButtonText = FAppStyle::Get().GetWidgetStyle<FTextBlockStyle>("SmallButtonText");
 	StyleSet->Set("RemoteControlExposeMenu.Heading", FTextBlockStyle(SmallButtonText)
@@ -120,6 +118,11 @@ void FRemoteControlExposeMenuStyle::Initialize()
 	FSlateStyleRegistry::RegisterSlateStyle(*StyleSet.Get());
 }
 
+#undef IMAGE_PLUGIN_BRUSH
+#undef IMAGE_PLUGIN_BRUSH_SVG
+#undef BOX_PLUGIN_BRUSH
+#undef RootToContentDir
+
 void FRemoteControlExposeMenuStyle::Shutdown()
 {
 	if (StyleSet.IsValid())
@@ -146,11 +149,3 @@ FString FRemoteControlExposeMenuStyle::InContent(const FString& RelativePath, co
 	static const FString ContentDir = IPluginManager::Get().FindPlugin(TEXT("RemoteControl"))->GetBaseDir() + TEXT("/Resources");
 	return (ContentDir / RelativePath) + Extension;
 }
-
-#undef IMAGE_PLUGIN_BRUSH
-#undef IMAGE_BRUSH
-#undef IMAGE_BRUSH_SVG
-#undef BOX_BRUSH
-#undef CORE_BOX_BRUSH
-#undef BOX_PLUGIN_BRUSH
-#undef CORE_FONT
diff --git a/Engine/Plugins/VirtualProduction/RemoteControl/Source/RemoteControlUI/Private/UI/RemoteControlPanelStyle.cpp b/Engine/Plugins/VirtualProduction/RemoteControl/Source/RemoteControlUI/Private/UI/RemoteControlPanelStyle.cpp
--- a/Engine/Plugins/VirtualProduction/RemoteControl/Source/RemoteControlUI/Private/UI/RemoteControlPanelStyle.cpp
+++ b/Engine/Plugins/VirtualProduction/RemoteControl/Source/RemoteControlUI/Private/UI/RemoteControlPanelStyle.cpp
@@ -11,14 +11,11 @@
 #include "Interfaces/IPluginManager.h"
 #include "SlateOptMacros.h"
 #include "Widgets/Input/SButton.h"
+#include "Styling/SlateStyleMacros.h"
 
 #define IMAGE_PLUGIN_BRUSH( RelativePath, ... ) FSlateImageBrush(FRemoteControlPanelStyle::InContent(RelativePath, ".png" ), __VA_ARGS__)
-#define IMAGE_BRUSH(RelativePath, ...) FSlateImageBrush(StyleSet->RootToContentDir(RelativePath, TEXT(".png")), __VA_ARGS__)
-#define IMAGE_BRUSH_SVG( RelativePath, ... ) FSlateVectorImageBrush(StyleSet->RootToContentDir(RelativePath, TEXT(".svg")), __VA_ARGS__)
-#define BOX_BRUSH(RelativePath, ...) FSlateBoxBrush(StyleSet->RootToContentDir(RelativePath, TEXT(".png")), __VA_ARGS__)
-#define CORE_BOX_BRUSH( RelativePath, ... ) FSlateBoxBrush(StyleSet->RootToCoreContentDir(RelativePath, TEXT(".png") ), __VA_ARGS__)
 #define BOX_PLUGIN_BRUSH( RelativePath, ... ) FSlateBoxBrush(FRemoteControlPanelStyle::InContent( RelativePath, ".png" ), __VA_ARGS__)
-#define DEFAULT_FONT(...) FAppStyle::Get().GetDefaultFontStyle(__VA_ARGS__)
+#define RootToContentDir StyleSet->RootToContentDir
 
 TSharedPtr<FSlateStyleSet> FRemoteControlPanelStyle::StyleSet;
 
@@ -248,8 +245,5 @@ void FRemoteControlPanelStyle::SetupPanelStyles(TSharedRef<FSlateStyleSet> InSty
 }
 
 #undef IMAGE_PLUGIN_BRUSH
-#undef IMAGE_BRUSH
-#undef IMAGE_BRUSH_SVG
-#undef BOX_BRUSH
 #undef BOX_PLUGIN_BRUSH
-#undef DEFAULT_FONT
+#undef RootToContentDir
diff --git a/Engine/Source/Developer/CollisionAnalyzer/Private/CollisionAnalyzerStyle.cpp b/Engine/Source/Developer/CollisionAnalyzer/Private/CollisionAnalyzerStyle.cpp
--- a/Engine/Source/Developer/CollisionAnalyzer/Private/CollisionAnalyzerStyle.cpp
+++ b/Engine/Source/Developer/CollisionAnalyzer/Private/CollisionAnalyzerStyle.cpp
@@ -4,11 +4,7 @@
 #include "Styling/SlateTypes.h"
 #include "Engine/GameEngine.h"
 #include "Styling/SlateStyleRegistry.h"
-
-#define IMAGE_BRUSH( RelativePath, ... ) FSlateImageBrush( RootToContentDir( RelativePath, TEXT(".png") ), __VA_ARGS__ )
-#define BOX_BRUSH( RelativePath, ... ) FSlateBoxBrush( RootToContentDir( RelativePath, TEXT(".png") ), __VA_ARGS__ )
-#define BORDER_BRUSH( RelativePath, ... ) FSlateBorderBrush( RootToContentDir( RelativePath, TEXT(".png") ), __VA_ARGS__ )
-#define DEFAULT_FONT(...) FCoreStyle::GetDefaultFontStyle(__VA_ARGS__)
+#include "Styling/SlateStyleMacros.h"
 
 TSharedPtr< FCollisionAnalyzerStyle::FStyle > FCollisionAnalyzerStyle::StyleInstance;
 
@@ -80,8 +76,3 @@ TSharedPtr< ISlateStyle > FCollisionAnalyzerStyle::Get()
 {
 	return StyleInstance;
 }
-
-#undef IMAGE_BRUSH
-#undef BOX_BRUSH
-#undef BORDER_BRUSH
-#undef DEFAULT_FONT
diff --git a/Engine/Source/Developer/SlateFileDialogs/Private/SlateFileDialogsStyles.cpp b/Engine/Source/Developer/SlateFileDialogs/Private/SlateFileDialogsStyles.cpp
--- a/Engine/Source/Developer/SlateFileDialogs/Private/SlateFileDialogsStyles.cpp
+++ b/Engine/Source/Developer/SlateFileDialogs/Private/SlateFileDialogsStyles.cpp
@@ -1,17 +1,15 @@
 // Copyright Epic Games, Inc. All Rights Reserved.
 
 #include "SlateFileDialogsStyles.h"
+#include "Styling/SlateStyleMacros.h"
 #include "Styling/SlateStyleRegistry.h"
 #include "Styling/SlateTypes.h"
 #include "Styling/CoreStyle.h"
+#include "Styling/SlateStyleMacros.h"
 
-TSharedPtr< FSlateStyleSet > FSlateFileDialogsStyle::StyleInstance = nullptr;
-
-#define IMAGE_BRUSH( RelativePath, ... ) FSlateImageBrush(Style->RootToContentDir(RelativePath, TEXT(".png")), __VA_ARGS__)
-#define BOX_BRUSH( RelativePath, ... ) FSlateBoxBrush(Style->RootToContentDir(RelativePath, TEXT(".png")), __VA_ARGS__)
-#define BORDER_BRUSH( RelativePath, ... ) FSlateBorderBrush(Style->RootToContentDir(RelativePath, TEXT(".png")), __VA_ARGS__)
-#define DEFAULT_FONT(...) FCoreStyle::GetDefaultFontStyle(__VA_ARGS__)
+#define RootToContentDir Style->RootToContentDir
 
+TSharedPtr< FSlateStyleSet > FSlateFileDialogsStyle::StyleInstance = nullptr;
 
 void FSlateFileDialogsStyle::Initialize()
 {
@@ -116,8 +114,4 @@ TSharedPtr< FSlateStyleSet > FSlateFileDialogsStyle::Create()
 	return Style;	
 }
 
-
-#undef IMAGE_BRUSH
-#undef BOX_BRUSH
-#undef BORDER_BRUSH
-#undef DEFAULT_FONT
+#undef RootToContentDir
diff --git a/Engine/Source/Developer/TraceTools/Private/TraceToolsStyle.cpp b/Engine/Source/Developer/TraceTools/Private/TraceToolsStyle.cpp
--- a/Engine/Source/Developer/TraceTools/Private/TraceToolsStyle.cpp
+++ b/Engine/Source/Developer/TraceTools/Private/TraceToolsStyle.cpp
@@ -24,8 +24,6 @@ TSharedPtr< FSlateStyleSet > FTraceToolsStyle::StyleSet = nullptr;
 
 FTextBlockStyle FTraceToolsStyle::NormalText;
 
-#define ICON_FONT(...) FSlateFontInfo(RootToContentDir("Fonts/FontAwesome", TEXT(".ttf")), __VA_ARGS__)
-
 // Const icon sizes
 static const FVector2D Icon8x8(8.0f, 8.0f);
 static const FVector2D Icon9x19(9.0f, 19.0f);
@@ -163,7 +161,7 @@ void FTraceToolsStyle::Initialize()
 	StyleSet->Set("FilterPresets.BackgroundBorder", new BOX_BRUSH("Common/DarkGroupBorder", FMargin(4.0f / 16.0f)));
 	StyleSet->Set("FilterPresets.WarningIcon", new CORE_IMAGE_BRUSH_SVG("Starship/Common/alert-circle", Icon40x40, FStyleColors::Warning));
 
-	StyleSet->Set("FontAwesome.9", ICON_FONT(9));
+	StyleSet->Set("FontAwesome.9", REGULAR_ICON_FONT(9));
 
 	StyleSet->Set("TraceControl.StartTrace", new CORE_IMAGE_BRUSH_SVG("Starship/Insights/TraceTools/TraceStart", Icon40x40));
 	StyleSet->Set("TraceControl.StartTrace.Small", new CORE_IMAGE_BRUSH_SVG("Starship/Insights/TraceTools/TraceStart", Icon20x20));
@@ -207,11 +205,8 @@ void FTraceToolsStyle::Initialize()
 	FSlateStyleRegistry::RegisterSlateStyle( *StyleSet.Get() );
 }
 
-#undef IMAGE_BRUSH
-#undef BOX_BRUSH
-#undef BORDER_BRUSH
-#undef DEFAULT_FONT
-#undef ICON_FONT
+#undef RootToContentDir
+#undef RootToCoreContentDir
 
 void FTraceToolsStyle::Shutdown()
 {
diff --git a/Engine/Source/Developer/TranslationEditor/Private/TranslationEditor.cpp b/Engine/Source/Developer/TranslationEditor/Private/TranslationEditor.cpp
--- a/Engine/Source/Developer/TranslationEditor/Private/TranslationEditor.cpp
+++ b/Engine/Source/Developer/TranslationEditor/Private/TranslationEditor.cpp
@@ -845,7 +845,9 @@ void FTranslationEditor::ChangeSourceFont()
 
 	if ( bOpened && NewFontFilename != "")
 	{
-		SourceFont = FSlateFontInfo(NewFontFilename, SourceFont.Size);
+		TSharedRef<FCompositeFont> SourceCompositeFont = MakeShared<FStandaloneCompositeFont>();
+		SourceCompositeFont->DefaultTypeface.AppendFont(TEXT("Regular"), NewFontFilename, EFontHinting::Default, EFontLoadingPolicy::LazyLoad);
+		SourceFont = FSlateFontInfo(SourceCompositeFont, SourceFont.Size);
 		RefreshUI();
 	}
 }
@@ -860,7 +862,9 @@ void FTranslationEditor::ChangeTranslationTargetFont()
 
 	if ( bOpened && NewFontFilename != "")
 	{
-		TranslationTargetFont = FSlateFontInfo(NewFontFilename, TranslationTargetFont.Size);
+		TSharedRef<FCompositeFont> TargetCompositeFont = MakeShared<FStandaloneCompositeFont>();
+		TargetCompositeFont->DefaultTypeface.AppendFont(TEXT("Regular"), NewFontFilename, EFontHinting::Default, EFontLoadingPolicy::LazyLoad);
+		TranslationTargetFont = FSlateFontInfo(TargetCompositeFont, TranslationTargetFont.Size);
 		RefreshUI();
 	}
 }
diff --git a/Engine/Source/Editor/EditorStyle/Private/SlateEditorStyle.cpp b/Engine/Source/Editor/EditorStyle/Private/SlateEditorStyle.cpp
--- a/Engine/Source/Editor/EditorStyle/Private/SlateEditorStyle.cpp
+++ b/Engine/Source/Editor/EditorStyle/Private/SlateEditorStyle.cpp
@@ -3,6 +3,7 @@
 #include "SlateEditorStyle.h"
 #include "Misc/CommandLine.h"
 #include "Styling/CoreStyle.h"
+#include "Styling/SlateStyleMacros.h"
 
 #include "Settings/EditorStyleSettings.h"
 
@@ -14,19 +15,6 @@
 
 BEGIN_SLATE_FUNCTION_BUILD_OPTIMIZATION
 
-#define IMAGE_BRUSH( RelativePath, ... ) FSlateImageBrush( RootToContentDir( RelativePath, TEXT(".png") ), __VA_ARGS__ )
-#define BOX_BRUSH( RelativePath, ... ) FSlateBoxBrush( RootToContentDir( RelativePath, TEXT(".png") ), __VA_ARGS__ )
-#define BORDER_BRUSH( RelativePath, ... ) FSlateBorderBrush( RootToContentDir( RelativePath, TEXT(".png") ), __VA_ARGS__ )
-
-#define CORE_IMAGE_BRUSH( RelativePath, ... ) FSlateImageBrush( RootToCoreContentDir( RelativePath, TEXT(".png") ), __VA_ARGS__ )
-#define CORE_BOX_BRUSH( RelativePath, ... ) FSlateBoxBrush( RootToCoreContentDir( RelativePath, TEXT(".png") ), __VA_ARGS__ )
-#define CORE_BORDER_BRUSH( RelativePath, ... ) FSlateBorderBrush( RootToCoreContentDir( RelativePath, TEXT(".png") ), __VA_ARGS__ )
-
-
-#define DEFAULT_FONT(...) FCoreStyle::GetDefaultFontStyle(__VA_ARGS__)
-#define ICON_FONT(...) FSlateFontInfo(RootToContentDir("Fonts/FontAwesome", TEXT(".ttf")), __VA_ARGS__)
-
-
 /* FSlateEditorStyle static initialization
  *****************************************************************************/
 
@@ -497,15 +485,15 @@ void FSlateEditorStyle::FStyle::SetupGeneralStyles()
 				);
 		}
 
-		Set("FontAwesome.7", ICON_FONT(7));
-		Set("FontAwesome.8", ICON_FONT(8));
-		Set("FontAwesome.9", ICON_FONT(9));
-		Set("FontAwesome.10", ICON_FONT(10));
-		Set("FontAwesome.11", ICON_FONT(11));
-		Set("FontAwesome.12", ICON_FONT(12));
-		Set("FontAwesome.14", ICON_FONT(14));
-		Set("FontAwesome.16", ICON_FONT(16));
-		Set("FontAwesome.18", ICON_FONT(18));
+		Set("FontAwesome.7", REGULAR_ICON_FONT(7));
+		Set("FontAwesome.8", REGULAR_ICON_FONT(8));
+		Set("FontAwesome.9", REGULAR_ICON_FONT(9));
+		Set("FontAwesome.10", REGULAR_ICON_FONT(10));
+		Set("FontAwesome.11", REGULAR_ICON_FONT(11));
+		Set("FontAwesome.12", REGULAR_ICON_FONT(12));
+		Set("FontAwesome.14", REGULAR_ICON_FONT(14));
+		Set("FontAwesome.16", REGULAR_ICON_FONT(16));
+		Set("FontAwesome.18", REGULAR_ICON_FONT(18));
 
 		/* Create a checkbox style for "ToggleButton" ... */
 		const FCheckBoxStyle ToggleButtonStyle = FCheckBoxStyle()
@@ -1457,7 +1445,7 @@ void FSlateEditorStyle::FStyle::SetupGeneralStyles()
 		);
 
 		Set("GraphBookmarkButtonText", FTextBlockStyle(NormalText)
-			.SetFont(ICON_FONT(16))
+			.SetFont(REGULAR_ICON_FONT(16))
 			.SetColorAndOpacity(FLinearColor(0.5f, 0.5f, 0.5f))
 			.SetShadowOffset(FVector2D::ZeroVector)
 		);
@@ -5833,7 +5821,7 @@ void FSlateEditorStyle::FStyle::SetupPersonaStyle()
 	//Find Results
 	{
 		Set("FindResults.FindInBlueprints", FTextBlockStyle(NormalText)
-			.SetFont(ICON_FONT(10))
+			.SetFont(REGULAR_ICON_FONT(10))
 			.SetColorAndOpacity(FLinearColor::White)
 		);
 
@@ -8181,10 +8169,4 @@ void FSlateEditorStyle::FStyle::SetupMyBlueprintStyles()
 	Set("MyBlueprint.DeleteEntry", new IMAGE_BRUSH("Icons/Edit/icon_Edit_Delete_40x", Icon16x16));
 }
 
-#undef IMAGE_BRUSH
-#undef BOX_BRUSH
-#undef BORDER_BRUSH
-#undef DEFAULT_FONT
-#undef ICON_FONT
-
 END_SLATE_FUNCTION_BUILD_OPTIMIZATION
diff --git a/Engine/Source/Editor/EditorStyle/Private/StarshipStyle.cpp b/Engine/Source/Editor/EditorStyle/Private/StarshipStyle.cpp
--- a/Engine/Source/Editor/EditorStyle/Private/StarshipStyle.cpp
+++ b/Engine/Source/Editor/EditorStyle/Private/StarshipStyle.cpp
@@ -18,9 +18,6 @@
 #include "Framework/Application/SlateApplication.h"
 #include "Framework/Docking/TabManager.h"
 
-
-#define ICON_FONT(...) FSlateFontInfo(RootToContentDir("Fonts/FontAwesome", TEXT(".ttf")), __VA_ARGS__)
-
 #define LOCTEXT_NAMESPACE "EditorStyle"
 
 FName FStarshipEditorStyle::StyleSetName = TEXT("EditorStyle");
@@ -594,15 +591,15 @@ void FStarshipEditorStyle::FStyle::SetupGeneralStyles()
 				);
 		}
 
-		Set("FontAwesome.7", ICON_FONT(7));
-		Set("FontAwesome.8", ICON_FONT(8));
-		Set("FontAwesome.9", ICON_FONT(9));
-		Set("FontAwesome.10", ICON_FONT(10));
-		Set("FontAwesome.11", ICON_FONT(11));
-		Set("FontAwesome.12", ICON_FONT(12));
-		Set("FontAwesome.14", ICON_FONT(14));
-		Set("FontAwesome.16", ICON_FONT(16));
-		Set("FontAwesome.18", ICON_FONT(18));
+		Set("FontAwesome.7", REGULAR_ICON_FONT(7));
+		Set("FontAwesome.8", REGULAR_ICON_FONT(8));
+		Set("FontAwesome.9", REGULAR_ICON_FONT(9));
+		Set("FontAwesome.10", REGULAR_ICON_FONT(10));
+		Set("FontAwesome.11", REGULAR_ICON_FONT(11));
+		Set("FontAwesome.12", REGULAR_ICON_FONT(12));
+		Set("FontAwesome.14", REGULAR_ICON_FONT(14));
+		Set("FontAwesome.16", REGULAR_ICON_FONT(16));
+		Set("FontAwesome.18", REGULAR_ICON_FONT(18));
 
 		/* Create a checkbox style for "ToggleButton" but with the images used by a normal checkbox (see "Checkbox" below) ... */
 		const FCheckBoxStyle CheckboxLookingToggleButtonStyle = FCheckBoxStyle()
@@ -5641,7 +5638,7 @@ void FStarshipEditorStyle::FStyle::SetupPersonaStyle()
 	//Find Results
 	{
 		Set("FindResults.FindInBlueprints", FTextBlockStyle(NormalText)
-			.SetFont(ICON_FONT(10))
+			.SetFont(REGULAR_ICON_FONT(10))
 			.SetColorAndOpacity(FLinearColor::White)
 		);
 	}
@@ -8169,8 +8166,6 @@ void FStarshipEditorStyle::FStyle::SetupSourceCodeStyles()
 	Set("SyntaxHighlight.SourceCode.Error", SourceCodeErrorText); 
 }
 
-#undef DEFAULT_FONT
-#undef ICON_FONT
 #undef LOCTEXT_NAMESPACE
 
 END_SLATE_FUNCTION_BUILD_OPTIMIZATION
diff --git a/Engine/Source/Editor/EditorWidgets/Private/EditorWidgetsStyle.cpp b/Engine/Source/Editor/EditorWidgets/Private/EditorWidgetsStyle.cpp
--- a/Engine/Source/Editor/EditorWidgets/Private/EditorWidgetsStyle.cpp
+++ b/Engine/Source/Editor/EditorWidgets/Private/EditorWidgetsStyle.cpp
@@ -11,15 +11,13 @@
 #include "Styling/SlateTypes.h"
 #include "Templates/TypeHash.h"
 #include "UObject/NameTypes.h"
+#include "Styling/SlateStyleMacros.h"
 
 class ISlateStyle;
 
 TSharedPtr<FSlateStyleSet> FEditorWidgetsStyle::StyleSet = nullptr;
 
-#define IMAGE_BRUSH( RelativePath, ... ) FSlateImageBrush( StyleSet->RootToContentDir( RelativePath, TEXT(".png") ), __VA_ARGS__ )
-#define BOX_BRUSH( RelativePath, ... ) FSlateBoxBrush( StyleSet->RootToContentDir( RelativePath, TEXT(".png") ), __VA_ARGS__ )
-#define BORDER_BRUSH( RelativePath, ... ) FSlateBorderBrush( StyleSet->RootToContentDir( RelativePath, TEXT(".png") ), __VA_ARGS__ )
-#define DEFAULT_FONT(...) FCoreStyle::GetDefaultFontStyle(__VA_ARGS__)
+#define RootToContentDir StyleSet->RootToContentDir
 
 void FEditorWidgetsStyle::Initialize()
 {
@@ -54,10 +52,7 @@ void FEditorWidgetsStyle::Initialize()
 	FSlateStyleRegistry::RegisterSlateStyle( *StyleSet.Get() );
 }
 
-#undef IMAGE_BRUSH
-#undef BOX_BRUSH
-#undef BORDER_BRUSH
-#undef DEFAULT_FONT
+#undef RootToContentDir
 
 void FEditorWidgetsStyle::Shutdown()
 {
diff --git a/Engine/Source/Editor/MassEntityDebugger/Private/MassDebuggerStyle.cpp b/Engine/Source/Editor/MassEntityDebugger/Private/MassDebuggerStyle.cpp
--- a/Engine/Source/Editor/MassEntityDebugger/Private/MassDebuggerStyle.cpp
+++ b/Engine/Source/Editor/MassEntityDebugger/Private/MassDebuggerStyle.cpp
@@ -14,10 +14,6 @@
 #include "Styling/StyleColors.h"
 #include "Textures/SlateIcon.h"
 
-
-#define TTF_CORE_FONT(RelativePath, ...) FSlateFontInfo(RootToCoreContentDir(RelativePath, TEXT(".ttf")), __VA_ARGS__)
-#define FONT(...) FSlateFontInfo(FCoreStyle::GetDefaultFont(), __VA_ARGS__)
-
 class FMassDebuggerStyleSet final : public FSlateStyleSet
 {
 public:
@@ -39,7 +35,7 @@ public:
 
 		const FTextBlockStyle& NormalText = FAppStyle::Get().GetWidgetStyle<FTextBlockStyle>("NormalText");
 
-		const FSlateFontInfo LargeFont(FONT(50, "Bold"));
+		const FSlateFontInfo LargeFont(DEFAULT_FONT("Bold", 50));
 
 		FTextBlockStyle StateTitle = FTextBlockStyle(NormalText)
 			.SetFont(LargeFont)
@@ -55,7 +51,7 @@ public:
 		Set("MassDebug.Label.Background", new FSlateRoundedBoxBrush(FStyleColors::Foreground, 4.0f));
 
 		Set("MassDebug.Label.Text", FTextBlockStyle(NormalText)
-			.SetFont(FSlateFontInfo(FONT(7, "Bold")))
+			.SetFont(FSlateFontInfo(DEFAULT_FONT("Bold", 7)))
 			.SetColorAndOpacity(FStyleColors::Background));
 	}
 };
diff --git a/Engine/Source/Editor/MassEntityEditor/Private/MassEditorStyle.cpp b/Engine/Source/Editor/MassEntityEditor/Private/MassEditorStyle.cpp
--- a/Engine/Source/Editor/MassEntityEditor/Private/MassEditorStyle.cpp
+++ b/Engine/Source/Editor/MassEntityEditor/Private/MassEditorStyle.cpp
@@ -1,18 +1,15 @@
 // Copyright Epic Games, Inc. All Rights Reserved.
 
 #include "MassEditorStyle.h"
+#include "Styling/SlateStyleMacros.h"
 #include "Styling/SlateStyleRegistry.h"
 #include "Styling/CoreStyle.h"
 #include "Styling/SlateTypes.h"
 #include "Styling/AppStyle.h"
 #include "Interfaces/IPluginManager.h"
 
-#define IMAGE_BRUSH( RelativePath, ... ) FSlateImageBrush( StyleSet->RootToContentDir( RelativePath, TEXT( ".png" ) ), __VA_ARGS__ )
 #define IMAGE_PLUGIN_BRUSH( RelativePath, ... ) FSlateImageBrush( FMeshEditorStyle::InContent( RelativePath, ".png" ), __VA_ARGS__ )
-#define BOX_BRUSH( RelativePath, ... ) FSlateBoxBrush( StyleSet->RootToContentDir( RelativePath, TEXT( ".png" ) ), __VA_ARGS__ )
-#define TTF_CORE_FONT( RelativePath, ... ) FSlateFontInfo( StyleSet->RootToCoreContentDir( RelativePath, TEXT( ".ttf" ) ), __VA_ARGS__ )
-
-#define DEFAULT_FONT(...) FCoreStyle::GetDefaultFontStyle(__VA_ARGS__)
+#define RootToContentDir StyleSet->RootToContentDir
 
 TSharedPtr<FSlateStyleSet> FMassEntityEditorStyle::StyleSet = nullptr;
 
@@ -47,13 +44,13 @@ void FMassEntityEditorStyle::Initialize()
 		StyleSet->Set("Mass.Icon", StateIcon);
 
 		FTextBlockStyle StateTitle = FTextBlockStyle(NormalText)
-			.SetFont(TTF_CORE_FONT("Fonts/Roboto-Bold", 12))
+			.SetFont(DEFAULT_FONT("Bold", 12))
 			.SetColorAndOpacity(FLinearColor(230.0f / 255.0f, 230.0f / 255.0f, 230.0f / 255.0f));
 		StyleSet->Set("Mass.State.Title", StateTitle);
 
 		FEditableTextBoxStyle StateTitleEditableText = FEditableTextBoxStyle()
 			.SetTextStyle(NormalText)
-			.SetFont(TTF_CORE_FONT("Fonts/Roboto-Bold", 10))
+			.SetFont(DEFAULT_FONT("Bold", 10))
 			.SetBackgroundImageNormal(BOX_BRUSH("Common/TextBox", FMargin(4.0f / 16.0f)))
 			.SetBackgroundImageHovered(BOX_BRUSH("Common/TextBox_Hovered", FMargin(4.0f / 16.0f)))
 			.SetBackgroundImageFocused(BOX_BRUSH("Common/TextBox_Hovered", FMargin(4.0f / 16.0f)))
@@ -69,13 +66,13 @@ void FMassEntityEditorStyle::Initialize()
 	// Task
 	{
 		FTextBlockStyle TaskTitle = FTextBlockStyle(NormalText)
-			.SetFont(TTF_CORE_FONT("Fonts/Roboto-Regular", 11))
+			.SetFont(DEFAULT_FONT("Regular", 11))
 			.SetColorAndOpacity(FLinearColor(230.0f / 255.0f, 230.0f / 255.0f, 230.0f / 255.0f));
 		StyleSet->Set("Mass.Task.Title", TaskTitle);
 
 		FEditableTextBoxStyle TaskTitleEditableText = FEditableTextBoxStyle()
 			.SetTextStyle(NormalText)
-			.SetFont(TTF_CORE_FONT("Fonts/Roboto-Regular", 9))
+			.SetFont(DEFAULT_FONT("Regular", 9))
 			.SetBackgroundImageNormal(BOX_BRUSH("Common/TextBox", FMargin(4.0f / 16.0f)))
 			.SetBackgroundImageHovered(BOX_BRUSH("Common/TextBox_Hovered", FMargin(4.0f / 16.0f)))
 			.SetBackgroundImageFocused(BOX_BRUSH("Common/TextBox_Hovered", FMargin(4.0f / 16.0f)))
@@ -91,7 +88,7 @@ void FMassEntityEditorStyle::Initialize()
 	// Details
 	{
 		FTextBlockStyle StateTitle = FTextBlockStyle(NormalText)
-			.SetFont(TTF_CORE_FONT("Fonts/Roboto-Regular", 10))
+			.SetFont(DEFAULT_FONT("Regular", 10))
 			.SetColorAndOpacity(FLinearColor(230.0f / 255.0f, 230.0f / 255.0f, 230.0f / 255.0f));
 		StyleSet->Set("Mass.Details", StateTitle);
 	}
@@ -110,6 +107,8 @@ void FMassEntityEditorStyle::Initialize()
 	FSlateStyleRegistry::RegisterSlateStyle(*StyleSet.Get());
 }
 
+#undef IMAGE_PLUGIN_BRUSH
+#undef RootToContentDir
 
 void FMassEntityEditorStyle::Shutdown()
 {
@@ -121,7 +120,6 @@ void FMassEntityEditorStyle::Shutdown()
 	}
 }
 
-
 FName FMassEntityEditorStyle::GetStyleSetName()
 {
 	static FName StyleName("MassEntityEditorStyle");
diff --git a/Engine/Source/Editor/MaterialEditor/Private/MaterialEditorStyle.cpp b/Engine/Source/Editor/MaterialEditor/Private/MaterialEditorStyle.cpp
--- a/Engine/Source/Editor/MaterialEditor/Private/MaterialEditorStyle.cpp
+++ b/Engine/Source/Editor/MaterialEditor/Private/MaterialEditorStyle.cpp
@@ -9,6 +9,11 @@
 #include "Misc/Paths.h"
 #include "Styling/SlateStyleRegistry.h"
 #include "Styling/StyleColors.h"
+#include "Styling/SlateStyleMacros.h"
+
+#define IMAGE_PLUGIN_BRUSH(RelativePath, ...) FSlateImageBrush(FSubstrateMaterialEditorStyle::InContent(RelativePath, ".png"), __VA_ARGS__)
+#define IMAGE_PLUGIN_BRUSH_SVG(RelativePath, ...) FSlateVectorImageBrush(FSubstrateMaterialEditorStyle::InContent(RelativePath, ".svg"), __VA_ARGS__)
+#define RootToContentDir Style->RootToContentDir
 
 TSharedPtr<FSlateStyleSet> FSubstrateMaterialEditorStyle::StyleInstance = nullptr;
 
@@ -56,22 +61,6 @@ const FSlateBrush* FSubstrateMaterialEditorStyle::GetBrush(const FName& InName)
 	return FSubstrateMaterialEditorStyle::Get().GetBrush(InName);
 }
 
-#define IMAGE_BRUSH(RelativePath, ... ) FSlateImageBrush(Style->RootToContentDir(RelativePath, TEXT(".png")), __VA_ARGS__)
-#define BOX_BRUSH(RelativePath, ... ) FSlateBoxBrush(Style->RootToContentDir(RelativePath, TEXT(".png")), __VA_ARGS__)
-#define BORDER_BRUSH(RelativePath, ... ) FSlateBorderBrush(Style->RootToContentDir(RelativePath, TEXT(".png")), __VA_ARGS__)
-
-#define IMAGE_BRUSH_SVG(RelativePath, ...) FSlateVectorImageBrush(Style->RootToContentDir(RelativePath, TEXT(".svg")), __VA_ARGS__)
-#define BOX_BRUSH_SVG(RelativePath, ...) FSlateVectorBoxBrush(Style->RootToContentDir(RelativePath, TEXT(".svg")), __VA_ARGS__)
-#define BORDER_BRUSH_SVG(RelativePath, ...) FSlateVectorBorderBrush(Style->RootToContentDir(RelativePath, TEXT(".svg")), __VA_ARGS__)
-
-#define CORE_IMAGE_BRUSH_SVG( RelativePath, ... ) FSlateVectorImageBrush(Style->RootToCoreContentDir(RelativePath, TEXT(".svg")), __VA_ARGS__)
-#define CORE_BOX_BRUSH_SVG( RelativePath, ... ) FSlateVectorBoxBrush(Style->RootToCoreContentDir(RelativePath, TEXT(".svg")), __VA_ARGS__)
-#define CORE_BORDER_BRUSH_SVG( RelativePath, ... ) FSlateVectorBorderBrush(Style->RootToCoreContentDir(RelativePath, TEXT(".svg")), __VA_ARGS__)
-
-#define DEFAULT_FONT(...) FCoreStyle::GetDefaultFontStyle(__VA_ARGS__)
-
-#define IMAGE_PLUGIN_BRUSH(RelativePath, ...) FSlateImageBrush(FSubstrateMaterialEditorStyle::InContent(RelativePath, ".png"), __VA_ARGS__)
-#define IMAGE_PLUGIN_BRUSH_SVG(RelativePath, ...) FSlateVectorImageBrush(FSubstrateMaterialEditorStyle::InContent(RelativePath, ".svg"), __VA_ARGS__)
 
 const FVector2D Icon8x8(8.f, 8.f);
 const FVector2D Icon12x12(12.f, 12.f);
@@ -367,19 +356,6 @@ void FSubstrateMaterialEditorStyle::SetupTextStyles(const TSharedRef<FSlateStyle
 		.SetColorAndOpacity(FStyleColors::SelectInactive));
 }
 
-#undef IMAGE_BRUSH
-#undef BOX_BRUSH
-#undef BORDER_BRUSH
-
-#undef IMAGE_BRUSH_SVG
-#undef BOX_BRUSH_SVG
-#undef BORDER_BRUSH_SVG
-
-#undef DEFAULT_FONT
-
 #undef IMAGE_PLUGIN_BRUSH
 #undef IMAGE_PLUGIN_BRUSH_SVG
-
-#undef CORE_IMAGE_BRUSH_SVG
-#undef CORE_BOX_BRUSH_SVG
-#undef CORE_BORDER_BRUSH_SVG
+#undef RootToContentDir
diff --git a/Engine/Source/Editor/MaterialEditor/Private/MaterialInstanceEditor.cpp b/Engine/Source/Editor/MaterialEditor/Private/MaterialInstanceEditor.cpp
--- a/Engine/Source/Editor/MaterialEditor/Private/MaterialInstanceEditor.cpp
+++ b/Engine/Source/Editor/MaterialEditor/Private/MaterialInstanceEditor.cpp
@@ -182,7 +182,7 @@ public:
 	virtual TSharedRef<SWidget> GenerateWidgetForColumn( const FName& ColumnName ) override
 	{
 		FText Entry;
-		FSlateFontInfo FontInfo = FSlateFontInfo( FPaths::EngineContentDir() / TEXT("Slate/Fonts/Roboto-Regular.ttf"), 9 );
+		FSlateFontInfo FontInfo = FSlateFontInfo(FCoreStyle::GetDefaultFont(), 9, "Regular");
 		if ( ColumnName == "Parent" )
 		{
 			if ( ParentIndex == 0 )
@@ -198,7 +198,7 @@ public:
 			else
 			{
 				Entry = NSLOCTEXT("UnrealEd", "Current", "Current");
-				FontInfo = FSlateFontInfo( FPaths::EngineContentDir() / TEXT("Slate/Fonts/Roboto-Bold.ttf"), 9 );
+				FontInfo = FSlateFontInfo(FCoreStyle::GetDefaultFont(), 9, "Bold");
 			}
 		}
 		else
@@ -206,7 +206,7 @@ public:
 			Entry = FText::FromString( WidgetInfo.Get()->GetName() );
 			if ( ParentIndex == -1 )
 			{
-				FontInfo = FSlateFontInfo( FPaths::EngineContentDir() / TEXT("Slate/Fonts/Roboto-Bold.ttf"), 9 );
+				FontInfo = FSlateFontInfo(FCoreStyle::GetDefaultFont(), 9, "Bold");
 			}
 		}
 		
diff --git a/Engine/Source/Editor/PIEPreviewDeviceProfileSelector/Private/PIEPreviewWindowCoreStyle.cpp b/Engine/Source/Editor/PIEPreviewDeviceProfileSelector/Private/PIEPreviewWindowCoreStyle.cpp
--- a/Engine/Source/Editor/PIEPreviewDeviceProfileSelector/Private/PIEPreviewWindowCoreStyle.cpp
+++ b/Engine/Source/Editor/PIEPreviewDeviceProfileSelector/Private/PIEPreviewWindowCoreStyle.cpp
@@ -8,6 +8,7 @@
 #include "Brushes/SlateBoxBrush.h"
 #include "Brushes/SlateColorBrush.h"
 #include "Brushes/SlateImageBrush.h"
+#include "Styling/SlateStyleMacros.h"
 #include "Styling/SlateStyleRegistry.h"
 #include "Styling/SlateTypes.h"
 #include "PIEPreviewDeviceSpecification.h"
@@ -16,9 +17,7 @@
 
 TSharedPtr< ISlateStyle > FPIEPreviewWindowCoreStyle::Instance = nullptr;
 
-#define IMAGE_BRUSH( RelativePath, ... ) FSlateImageBrush(Style->RootToContentDir(RelativePath, TEXT(".png")), __VA_ARGS__)
-#define DEFAULT_FONT(...) FCoreStyle::GetDefaultFontStyle(__VA_ARGS__)
-#define BOX_BRUSH( RelativePath, ... ) FSlateBoxBrush(Style->RootToContentDir( RelativePath, TEXT(".png") ), __VA_ARGS__)
+#define RootToContentDir Style->RootToContentDir
 
 TSharedRef<class ISlateStyle> FPIEPreviewWindowCoreStyle::Create(const FName& InStyleSetName)
 {
@@ -48,10 +47,6 @@ TSharedRef<class ISlateStyle> FPIEPreviewWindowCoreStyle::Create(const FName& In
 	return Style;
 }
 
-#undef IMAGE_BRUSH
-#undef DEFAULT_FONT
-#undef BOX_BRUSH
-
 void FPIEPreviewWindowCoreStyle::InitializePIECoreStyle()
 {
 	if (!Instance.IsValid())
@@ -60,4 +55,7 @@ void FPIEPreviewWindowCoreStyle::InitializePIECoreStyle()
 		FSlateStyleRegistry::RegisterSlateStyle(*Instance.Get());
 	}
 }
+
+#undef RootToContentDir
+
 #endif
\ No newline at end of file
diff --git a/Engine/Source/Editor/UnrealEd/Private/AssetThumbnail.cpp b/Engine/Source/Editor/UnrealEd/Private/AssetThumbnail.cpp
--- a/Engine/Source/Editor/UnrealEd/Private/AssetThumbnail.cpp
+++ b/Engine/Source/Editor/UnrealEd/Private/AssetThumbnail.cpp
@@ -944,7 +944,7 @@ private:
 			.Visibility(this, &SAssetThumbnail::GetStatusOverflowVisibility)
 			[
 				SNew(STextBlock)
-				.Font(FSlateFontInfo( FPaths::EngineContentDir() / TEXT("Slate/Fonts/Roboto-Bold.ttf"), 9))
+				.Font(FSlateFontInfo(FCoreStyle::GetDefaultFont(), 9, "Bold"))
 				.Text(this, &SAssetThumbnail::GetStatusOverflowText)
 			];
 	}
diff --git a/Engine/Source/Editor/VREditor/Private/VREditorStyle.cpp b/Engine/Source/Editor/VREditor/Private/VREditorStyle.cpp
--- a/Engine/Source/Editor/VREditor/Private/VREditorStyle.cpp
+++ b/Engine/Source/Editor/VREditor/Private/VREditorStyle.cpp
@@ -15,10 +15,13 @@
 #include "Rendering/SlateRenderer.h"
 #include "Styling/AppStyle.h"
 #include "Styling/CoreStyle.h"
+#include "Styling/SlateStyleMacros.h"
 #include "Styling/SlateStyleRegistry.h"
 #include "Styling/SlateTypes.h"
 #include "Templates/TypeHash.h"
 
+#define RootToContentDir Style->RootToContentDir
+
 class ISlateStyle;
 
 TSharedPtr< FSlateStyleSet > FVREditorStyle::VREditorStyleInstance = NULL;
@@ -51,11 +54,6 @@ FName FVREditorStyle::GetNumpadStyleSetName()
 	return StyleSetName;
 }
 
-#define IMAGE_BRUSH( RelativePath, ... ) FSlateImageBrush( Style->RootToContentDir( RelativePath, TEXT(".png") ), __VA_ARGS__ )
-#define BOX_BRUSH( RelativePath, ... ) FSlateBoxBrush( Style->RootToContentDir( RelativePath, TEXT(".png") ), __VA_ARGS__ )
-#define BORDER_BRUSH( RelativePath, ... ) FSlateBorderBrush( Style->RootToContentDir( RelativePath, TEXT(".png") ), __VA_ARGS__ )
-#define DEFAULT_FONT(...) FCoreStyle::GetDefaultFontStyle(__VA_ARGS__)
-
 const FVector2D Icon14x14(14.0f, 14.0f);
 const FVector2D Icon16x16(16.0f, 16.0f);
 const FVector2D Icon20x20(20.0f, 20.0f);
@@ -251,8 +249,4 @@ const ISlateStyle& FVREditorStyle::Get()
 	return *VREditorStyleInstance;
 }
 
-
-#undef IMAGE_BRUSH
-#undef BOX_BRUSH
-#undef BORDER_BRUSH
-#undef DEFAULT_FONT
+#undef RootToContentDir
diff --git a/Engine/Source/Programs/CrashReportClient/Private/CrashReportClientStyle.cpp b/Engine/Source/Programs/CrashReportClient/Private/CrashReportClientStyle.cpp
--- a/Engine/Source/Programs/CrashReportClient/Private/CrashReportClientStyle.cpp
+++ b/Engine/Source/Programs/CrashReportClient/Private/CrashReportClientStyle.cpp
@@ -8,9 +8,12 @@
 #include "Math/Vector2D.h"
 #include "CrashReportClientApp.h"
 #include "Styling/StyleColors.h"
+#include "Styling/SlateStyleMacros.h"
 
 #if !CRASH_REPORT_UNATTENDED_ONLY
 
+#define RootToContentDir ContentFromEngine
+
 #include "Styling/SlateStyleRegistry.h"
 #include "Styling/SlateTypes.h"
 #include "Styling/CoreStyle.h"
@@ -33,11 +36,6 @@ void FCrashReportClientStyle::Shutdown()
 	StyleSet.Reset();
 }
 
-#define IMAGE_BRUSH( RelativePath, ... ) FSlateImageBrush(ContentFromEngine(TEXT(RelativePath), TEXT(".png")), __VA_ARGS__)
-#define BOX_BRUSH( RelativePath, ... ) FSlateBoxBrush(ContentFromEngine(TEXT(RelativePath), TEXT(".png")), __VA_ARGS__)
-#define BORDER_BRUSH( RelativePath, ... ) FSlateBorderBrush(ContentFromEngine(RelativePath, TEXT(".png")), __VA_ARGS__)
-#define DEFAULT_FONT(...) FCoreStyle::GetDefaultFontStyle(__VA_ARGS__)
-
 namespace
 {
 	FString ContentFromEngine(const FString& RelativePath, const TCHAR* Extension)
@@ -124,14 +122,11 @@ TSharedRef< FSlateStyleSet > FCrashReportClientStyle::Create()
 	return StyleRef;
 }
 
-#undef IMAGE_BRUSH
-#undef BOX_BRUSH
-#undef BORDER_BRUSH
-#undef DEFAULT_FONT
-
 const ISlateStyle& FCrashReportClientStyle::Get()
 {
 	return *StyleSet;
 }
 
+#undef RootToContentDir
+
 #endif // !CRASH_REPORT_UNATTENDED_ONLY
diff --git a/Engine/Source/Programs/LiveCodingConsole/Private/LiveCodingConsoleStyle.cpp b/Engine/Source/Programs/LiveCodingConsole/Private/LiveCodingConsoleStyle.cpp
--- a/Engine/Source/Programs/LiveCodingConsole/Private/LiveCodingConsoleStyle.cpp
+++ b/Engine/Source/Programs/LiveCodingConsole/Private/LiveCodingConsoleStyle.cpp
@@ -6,10 +6,13 @@
 #include "Misc/Paths.h"
 #include "Math/Color.h"
 #include "Math/Vector2D.h"
+#include "Styling/SlateStyleMacros.h"
 #include "Styling/SlateStyleRegistry.h"
 #include "Styling/SlateTypes.h"
 #include "Styling/CoreStyle.h"
 
+#define RootToContentDir ContentFromEngine
+
 TSharedPtr< FSlateStyleSet > FLiveCodingConsoleStyle::StyleSet = nullptr;
 
 void FLiveCodingConsoleStyle::Initialize()
@@ -28,11 +31,6 @@ void FLiveCodingConsoleStyle::Shutdown()
 	StyleSet.Reset();
 }
 
-#define IMAGE_BRUSH( RelativePath, ... ) FSlateImageBrush(ContentFromEngine(TEXT(RelativePath), TEXT(".png")), __VA_ARGS__)
-#define BOX_BRUSH( RelativePath, ... ) FSlateBoxBrush(ContentFromEngine(TEXT(RelativePath), TEXT(".png")), __VA_ARGS__)
-#define BORDER_BRUSH( RelativePath, ... ) FSlateBorderBrush(ContentFromEngine(RelativePath, TEXT(".png")), __VA_ARGS__)
-#define DEFAULT_FONT(...) FCoreStyle::GetDefaultFontStyle(__VA_ARGS__)
-
 namespace
 {
 	FString ContentFromEngine(const FString& RelativePath, const TCHAR* Extension)
@@ -98,14 +96,9 @@ TSharedRef< FSlateStyleSet > FLiveCodingConsoleStyle::Create()
 	return StyleRef;
 }
 
-#undef IMAGE_BRUSH
-#undef BOX_BRUSH
-#undef BORDER_BRUSH
-#undef DEFAULT_FONT
-
 const ISlateStyle& FLiveCodingConsoleStyle::Get()
 {
 	return *StyleSet;
 }
 
-
+#undef RootToContentDir
diff --git a/Engine/Source/Programs/ZenDashboard/Private/ZenDashboardStyle.cpp b/Engine/Source/Programs/ZenDashboard/Private/ZenDashboardStyle.cpp
--- a/Engine/Source/Programs/ZenDashboard/Private/ZenDashboardStyle.cpp
+++ b/Engine/Source/Programs/ZenDashboard/Private/ZenDashboardStyle.cpp
@@ -6,10 +6,17 @@
 #include "Misc/Paths.h"
 #include "Math/Color.h"
 #include "Math/Vector2D.h"
+#include "Styling/SlateStyleMacros.h"
 #include "Styling/SlateStyleRegistry.h"
 #include "Styling/SlateTypes.h"
 #include "Styling/CoreStyle.h"
 
+#define EDITOR_IMAGE_BRUSH(RelativePath, ...) IMAGE_BRUSH("../Editor/Slate/" RelativePath, __VA_ARGS__)
+#define EDITOR_IMAGE_BRUSH_SVG(RelativePath, ...) IMAGE_BRUSH_SVG("../Editor/Slate/" RelativePath, __VA_ARGS__)
+#define EDITOR_BOX_BRUSH(RelativePath, ...) BOX_BRUSH("../Editor/Slate/" RelativePath, __VA_ARGS__)
+#define EDITOR_BORDER_BRUSH(RelativePath, ...) BORDER_BRUSH("../Editor/Slate/" RelativePath, __VA_ARGS__)
+#define RootToContentDir ContentFromEngine
+
 TSharedPtr< FSlateStyleSet > FZenDashboardStyle::StyleSet = nullptr;
 
 void FZenDashboardStyle::Initialize()
@@ -28,17 +35,6 @@ void FZenDashboardStyle::Shutdown()
 	StyleSet.Reset();
 }
 
-#define IMAGE_BRUSH( RelativePath, ... ) FSlateImageBrush(ContentFromEngine(TEXT(RelativePath), TEXT(".png")), __VA_ARGS__)
-#define IMAGE_BRUSH_SVG( RelativePath, ... ) FSlateVectorImageBrush(ContentFromEngine(TEXT(RelativePath), TEXT(".svg")), __VA_ARGS__)
-#define BOX_BRUSH( RelativePath, ... ) FSlateBoxBrush(ContentFromEngine(TEXT(RelativePath), TEXT(".png")), __VA_ARGS__)
-#define BORDER_BRUSH( RelativePath, ... ) FSlateBorderBrush(ContentFromEngine(RelativePath, TEXT(".png")), __VA_ARGS__)
-#define DEFAULT_FONT(...) FCoreStyle::GetDefaultFontStyle(__VA_ARGS__)
-
-#define EDITOR_IMAGE_BRUSH(RelativePath, ...) IMAGE_BRUSH("../Editor/Slate/" RelativePath, __VA_ARGS__)
-#define EDITOR_IMAGE_BRUSH_SVG(RelativePath, ...) IMAGE_BRUSH_SVG("../Editor/Slate/" RelativePath, __VA_ARGS__)
-#define EDITOR_BOX_BRUSH(RelativePath, ...) BOX_BRUSH("../Editor/Slate/" RelativePath, __VA_ARGS__)
-#define EDITOR_BORDER_BRUSH(RelativePath, ...) BORDER_BRUSH("../Editor/Slate/" RelativePath, __VA_ARGS__)
-
 namespace
 {
 	FString ContentFromEngine(const FString& RelativePath, const TCHAR* Extension)
@@ -108,14 +104,13 @@ TSharedRef< FSlateStyleSet > FZenDashboardStyle::Create()
 	return StyleRef;
 }
 
-#undef IMAGE_BRUSH
-#undef BOX_BRUSH
-#undef BORDER_BRUSH
-#undef DEFAULT_FONT
-
 const ISlateStyle& FZenDashboardStyle::Get()
 {
 	return *StyleSet;
 }
 
-
+#undef EDITOR_IMAGE_BRUSH
+#undef EDITOR_IMAGE_BRUSH_SVG
+#undef EDITOR_BOX_BRUSH
+#undef EDITOR_BORDER_BRUSH
+#undef RootToContentDir
diff --git a/Engine/Source/Runtime/AppFramework/Private/Framework/Testing/TestStyle.cpp b/Engine/Source/Runtime/AppFramework/Private/Framework/Testing/TestStyle.cpp
--- a/Engine/Source/Runtime/AppFramework/Private/Framework/Testing/TestStyle.cpp
+++ b/Engine/Source/Runtime/AppFramework/Private/Framework/Testing/TestStyle.cpp
@@ -4,6 +4,7 @@
 #include "Brushes/SlateBorderBrush.h"
 #include "Brushes/SlateBoxBrush.h"
 #include "Brushes/SlateImageBrush.h"
+#include "Styling/SlateStyleMacros.h"
 #include "Styling/SlateStyleRegistry.h"
 #include "Styling/SlateStyle.h"
 #include "Styling/SlateTypes.h"
@@ -12,6 +13,7 @@
 
 #if !UE_BUILD_SHIPPING
 
+#define RootToContentDir Style->RootToContentDir
 
 TSharedPtr< ISlateStyle > FTestStyle::Instance = nullptr;
 
@@ -40,12 +42,6 @@ void FTestStyle::SetStyle( const TSharedRef< ISlateStyle >& NewStyle )
 	}
 }
 
-#define IMAGE_BRUSH( RelativePath, ... ) FSlateImageBrush( Style->RootToContentDir( RelativePath, TEXT(".png") ), __VA_ARGS__ )
-#define IMAGE_BRUSH_SVG( RelativePath, ... ) FSlateVectorImageBrush( Style->RootToContentDir(RelativePath, TEXT(".svg")), __VA_ARGS__)
-#define BOX_BRUSH( RelativePath, ... ) FSlateBoxBrush( Style->RootToContentDir( RelativePath, TEXT(".png") ), __VA_ARGS__ )
-#define BORDER_BRUSH( RelativePath, ... ) FSlateBorderBrush( Style->RootToContentDir( RelativePath, TEXT(".png") ), __VA_ARGS__ )
-#define DEFAULT_FONT(...) FCoreStyle::GetDefaultFontStyle(__VA_ARGS__)
-
 TSharedRef< ISlateStyle > FTestStyle::Create()
 {
 	TSharedRef< FSlateStyleSet > Style = MakeShareable( new FSlateStyleSet("TestStyle") );
@@ -406,9 +402,6 @@ TSharedRef< ISlateStyle > FTestStyle::Create()
 	return Style;
 }
 
-#undef IMAGE_BRUSH
-#undef BOX_BRUSH
-#undef BORDER_BRUSH
-#undef DEFAULT_FONT
+#undef RootToContentDir
 
 #endif // #if !UE_BUILD_SHIPPING
diff --git a/Engine/Source/Runtime/GameMenuBuilder/Private/GameMenuBuilderStyle.cpp b/Engine/Source/Runtime/GameMenuBuilder/Private/GameMenuBuilderStyle.cpp
--- a/Engine/Source/Runtime/GameMenuBuilder/Private/GameMenuBuilderStyle.cpp
+++ b/Engine/Source/Runtime/GameMenuBuilder/Private/GameMenuBuilderStyle.cpp
@@ -2,6 +2,7 @@
 
 #include "GameMenuBuilderStyle.h"
 #include "Modules/ModuleManager.h"
+#include "Styling/SlateStyleMacros.h"
 #include "Styling/SlateStyleRegistry.h"
 #include "Styling/SlateTypes.h"
 #include "Styling/CoreStyle.h"
@@ -9,8 +10,6 @@
 #include "Slate/SlateGameResources.h"
 #include "GameMenuBuilderModule.h"
 
-
-
 TSharedPtr< FSlateStyleSet > FGameMenuBuilderStyle::SimpleStyleInstance = NULL;
 
 void FGameMenuBuilderStyle::Initialize(const FString StyleName)
@@ -39,11 +38,6 @@ FName FGameMenuBuilderStyle::GetStyleSetName()
 	return StyleSetName;
 }
 
-#define IMAGE_BRUSH( RelativePath, ... ) FSlateImageBrush( FPaths::ProjectContentDir() / "Slate"/ RelativePath + TEXT(".png"), __VA_ARGS__ )
-#define BOX_BRUSH( RelativePath, ... ) FSlateBoxBrush( FPaths::ProjectContentDir() / "Slate"/ RelativePath + TEXT(".png"), __VA_ARGS__ )
-#define BORDER_BRUSH( RelativePath, ... ) FSlateBorderBrush( FPaths::ProjectContentDir() / "Slate"/ RelativePath + TEXT(".png"), __VA_ARGS__ )
-#define DEFAULT_FONT(...) FCoreStyle::GetDefaultFontStyle(__VA_ARGS__)
-
 TSharedRef< FSlateStyleSet > FGameMenuBuilderStyle::Create(const FString StyleName)
 {
 	TSharedRef<FSlateGameResources> StyleRef = FSlateGameResources::New(GetStyleSetName(), *StyleName);
@@ -66,11 +60,6 @@ TSharedRef< FSlateStyleSet > FGameMenuBuilderStyle::Create(const FString StyleNa
 	return StyleRef;
 }
 
-#undef IMAGE_BRUSH
-#undef BOX_BRUSH
-#undef BORDER_BRUSH
-#undef DEFAULT_FONT
-
 void FGameMenuBuilderStyle::ReloadTextures()
 {
 	FSlateApplication::Get().GetRenderer()->ReloadTextureResources();
diff --git a/Engine/Source/Runtime/SlateCore/Private/Styling/CoreStyle.cpp b/Engine/Source/Runtime/SlateCore/Private/Styling/CoreStyle.cpp
--- a/Engine/Source/Runtime/SlateCore/Private/Styling/CoreStyle.cpp
+++ b/Engine/Source/Runtime/SlateCore/Private/Styling/CoreStyle.cpp
@@ -6,12 +6,15 @@
 #include "Brushes/SlateBoxBrush.h"
 #include "Brushes/SlateColorBrush.h"
 #include "Brushes/SlateImageBrush.h"
+#include "Styling/SlateStyleMacros.h"
 #include "Styling/SlateStyleRegistry.h"
 #include "Styling/SlateStyle.h"
 #include "Styling/SlateTypes.h"
 #include "Fonts/LegacySlateFontInfoCache.h"
 #include "Misc/CommandLine.h"
 
+#define RootToContentDir Style->RootToContentDir
+
 
 /* Static initialization
  *****************************************************************************/
@@ -68,12 +71,24 @@ TSharedRef<const FCompositeFont> FCoreStyle::GetDefaultFont()
 	return FLegacySlateFontInfoCache::Get().GetDefaultFont();
 }
 
+TSharedRef<const FCompositeFont> FCoreStyle::GetIconFont()
+{
+#if WITH_EDITOR
+	return FLegacySlateFontInfoCache::Get().GetCompositeFont(TEXT("Editor/Slate/Fonts/FontAwesome.ttf"), EFontHinting::Default).ToSharedRef();
+#else
+	return GetDefaultFont();
+#endif
+}
 
 FSlateFontInfo FCoreStyle::GetDefaultFontStyle(const FName InTypefaceFontName, const float InSize, const FFontOutlineSettings& InOutlineSettings)
 {
 	return FSlateFontInfo(GetDefaultFont(), InSize, InTypefaceFontName, InOutlineSettings);
 }
 
+FSlateFontInfo FCoreStyle::GetRegularIconFontStyle(const float InSize, const FFontOutlineSettings& InOutlineSettings)
+{
+	return FSlateFontInfo(GetIconFont(), InSize, "Regular", InOutlineSettings);
+}
 
 void FCoreStyle::ResetToDefault( )
 {
@@ -129,11 +144,6 @@ void FCoreStyle::SetFocusBrush(FSlateBrush* NewBrush)
 
 UE_DISABLE_OPTIMIZATION_SHIP
 
-#define IMAGE_BRUSH( RelativePath, ... ) FSlateImageBrush(Style->RootToContentDir(RelativePath, TEXT(".png")), __VA_ARGS__)
-#define BOX_BRUSH( RelativePath, ... ) FSlateBoxBrush(Style->RootToContentDir(RelativePath, TEXT(".png")), __VA_ARGS__)
-#define BORDER_BRUSH( RelativePath, ... ) FSlateBorderBrush(Style->RootToContentDir(RelativePath, TEXT(".png")), __VA_ARGS__)
-#define DEFAULT_FONT(...) FCoreStyle::GetDefaultFontStyle(__VA_ARGS__)
-
 // We no longer use this style. it is kept around for reference only
 #define ALLOW_LEGACY_CORE_STYLE 0
 
@@ -142,7 +152,7 @@ TSharedRef<ISlateStyle> FCoreStyle::Create( const FName& InStyleSetName )
 	TSharedRef<FSlateCoreStyle> Style = MakeShareable(new FSlateCoreStyle(InStyleSetName));
 	Style->SetContentRoot(FPaths::EngineContentDir() / TEXT("Slate"));
 
-	const FString CanaryPath = Style->RootToContentDir(TEXT("Checkerboard"), TEXT(".png"));
+	const FString CanaryPath = RootToContentDir(TEXT("Checkerboard"), TEXT(".png"));
 
 	if (!FPaths::FileExists(CanaryPath))
 	{
@@ -1340,12 +1350,10 @@ TSharedRef<ISlateStyle> FCoreStyle::Create( const FName& InStyleSetName )
 	return Style;
 
 }
+
 UE_ENABLE_OPTIMIZATION_SHIP
 
-#undef IMAGE_BRUSH
-#undef BOX_BRUSH
-#undef BORDER_BRUSH
-#undef DEFAULT_FONT
+#undef RootToContentDir
 
 
 const TSharedPtr<FSlateDynamicImageBrush> FCoreStyle::GetDynamicImageBrush( FName BrushTemplate, FName TextureName, const ANSICHAR* Specifier )
diff --git a/Engine/Source/Runtime/SlateCore/Private/Styling/StarshipCoreStyle.cpp b/Engine/Source/Runtime/SlateCore/Private/Styling/StarshipCoreStyle.cpp
--- a/Engine/Source/Runtime/SlateCore/Private/Styling/StarshipCoreStyle.cpp
+++ b/Engine/Source/Runtime/SlateCore/Private/Styling/StarshipCoreStyle.cpp
@@ -28,22 +28,20 @@ TSharedPtr< ISlateStyle > FStarshipCoreStyle::Instance = nullptr;
 
 using namespace CoreStyleConstants;
 
-#define FONT(...) FSlateFontInfo(FCoreStyle::GetDefaultFont(), __VA_ARGS__)
-
 TUniquePtr<FStyleFonts> FStyleFonts::Instance = nullptr;
 
 FStyleFonts::FStyleFonts()
-	: Normal(FONT(10, "Regular"))
-	, NormalItalic(FONT(10, "Italic"))
-	, NormalBold(FONT(10, "Bold"))
-	, NormalBoldItalic(FONT(10, "BoldItalic"))
-	, Small(FONT(8,  "Regular"))
-	, SmallBold(FONT(8,  "Bold"))
-	, Large(FONT(14, "Regular"))
-	, LargeBold(FONT(14, "Bold"))
-	, HeadingMedium(FONT(33, "BoldCondensed"))
-	, HeadingSmall(FONT(21, "BoldCondensed"))
-	, HeadingExtraSmall(FONT(15, "BoldCondensed"))
+	: Normal(DEFAULT_FONT("Regular", 10))
+	, NormalItalic(DEFAULT_FONT("Italic", 10))
+	, NormalBold(DEFAULT_FONT("Bold", 10))
+	, NormalBoldItalic(DEFAULT_FONT("BoldItalic", 10))
+	, Small(DEFAULT_FONT("Regular", 8))
+	, SmallBold(DEFAULT_FONT("Bold", 8))
+	, Large(DEFAULT_FONT("Regular", 14))
+	, LargeBold(DEFAULT_FONT("Bold", 14))
+	, HeadingMedium(DEFAULT_FONT("BoldCondensed", 33))
+	, HeadingSmall(DEFAULT_FONT("BoldCondensed", 21))
+	, HeadingExtraSmall(DEFAULT_FONT("BoldCondensed", 15))
 {
 };
 
@@ -1346,7 +1344,7 @@ void FStarshipCoreStyle::SetupTextStyles(TSharedRef<FStyle>& Style)
 	}
 
 	const FTextBlockStyle HintText = FTextBlockStyle(NormalText)
-		.SetFont(FONT(8, "Italic"))
+		.SetFont(DEFAULT_FONT("Italic", 8))
 		.SetColorAndOpacity(FSlateColor::UseSubduedForeground());
 	Style->Set("HintText", HintText);
 }
@@ -2342,7 +2340,7 @@ void FStarshipCoreStyle::SetupMultiboxStyles(TSharedRef<FStyle>& Style)
 		Style->Set("Menu.Keybinding", FTextBlockStyle(NormalText).SetFont(FStyleFonts::Get().Small));
 
 
-		FSlateFontInfo XSFont(FONT(7, "Bold"));
+		FSlateFontInfo XSFont(DEFAULT_FONT("Bold", 7));
 		XSFont.LetterSpacing =  250;
 
 		Style->Set("Menu.Heading",
@@ -2470,4 +2468,3 @@ void FStarshipCoreStyle::SetupMultiboxStyles(TSharedRef<FStyle>& Style)
 }
 
 #undef RootToContentDir
-#undef FONT
diff --git a/Engine/Source/Runtime/SlateCore/Public/Styling/CoreStyle.h b/Engine/Source/Runtime/SlateCore/Public/Styling/CoreStyle.h
--- a/Engine/Source/Runtime/SlateCore/Public/Styling/CoreStyle.h
+++ b/Engine/Source/Runtime/SlateCore/Public/Styling/CoreStyle.h
@@ -44,9 +44,15 @@ public:
 	/** Get the default font for Slate */
 	static SLATECORE_API TSharedRef<const FCompositeFont> GetDefaultFont();
 
-	/** Get a font style using the default for for Slate */
+	/** Get the icon font for Slate (works only in Editor, otherwise the DefaultFont is returned) */
+	static SLATECORE_API TSharedRef<const FCompositeFont> GetIconFont();
+
+	/** Get a font style using the default font for Slate */
 	static SLATECORE_API FSlateFontInfo GetDefaultFontStyle(const FName InTypefaceFontName, const float InSize, const FFontOutlineSettings& InOutlineSettings = FFontOutlineSettings());
 
+	/** Get a font style using the icon font for Slate (always using "Regular" typeface */
+	static SLATECORE_API FSlateFontInfo GetRegularIconFontStyle(const float InSize, const FFontOutlineSettings& InOutlineSettings = FFontOutlineSettings());
+
 	static SLATECORE_API void ResetToDefault( );
 
 	/** Used to override the default selection colors */
diff --git a/Engine/Source/Runtime/SlateCore/Public/Styling/SlateStyleMacros.h b/Engine/Source/Runtime/SlateCore/Public/Styling/SlateStyleMacros.h
--- a/Engine/Source/Runtime/SlateCore/Public/Styling/SlateStyleMacros.h
+++ b/Engine/Source/Runtime/SlateCore/Public/Styling/SlateStyleMacros.h
@@ -18,4 +18,5 @@
 #define CORE_BOX_BRUSH_SVG( RelativePath, ... ) FSlateVectorBoxBrush(RootToCoreContentDir(RelativePath, TEXT(".svg")), __VA_ARGS__)
 #define CORE_BORDER_BRUSH_SVG( RelativePath, ... ) FSlateVectorBorderBrush(RootToCoreContentDir(RelativePath, TEXT(".svg")), __VA_ARGS__)
 
-#define DEFAULT_FONT(...) FCoreStyle::GetDefaultFontStyle(__VA_ARGS__)
\ No newline at end of file
+#define DEFAULT_FONT(...) FCoreStyle::GetDefaultFontStyle(__VA_ARGS__)
+#define REGULAR_ICON_FONT(...) FCoreStyle::GetRegularIconFontStyle(__VA_ARGS__)
\ No newline at end of file
