Check that UMassEnttiySettings is non-null
- To satisfy PVS Studio.
- Wasn't checked before, but was assumed to be valid and immediately dereferenced.
- Related: 46325358

diff --git a/Engine/Plugins/Runtime/MassGameplay/Source/MassSimulation/Private/MassSimulationSubsystem.cpp b/Engine/Plugins/Runtime/MassGameplay/Source/MassSimulation/Private/MassSimulationSubsystem.cpp
--- a/Engine/Plugins/Runtime/MassGameplay/Source/MassSimulation/Private/MassSimulationSubsystem.cpp
+++ b/Engine/Plugins/Runtime/MassGameplay/Source/MassSimulation/Private/MassSimulationSubsystem.cpp
@@ -187,12 +187,13 @@ void UMassSimulationSubsystem::HandleSimulationTickingEnabledCVarChange(IConsole
 void UMassSimulationSubsystem::RebuildTickPipeline()
 {
 	UMassEntitySettings* Settings = GetMutableDefault<UMassEntitySettings>();
+	check(Settings);
 	TConstArrayView<FMassProcessingPhaseConfig> ProcessingPhasesConfig = Settings->GetProcessingPhasesConfig();
 	FString DependencyGraphFileName;
 
 #if WITH_EDITOR
 	const UWorld* World = GetWorld();
-	if (World != nullptr && Settings != nullptr && !Settings->DumpDependencyGraphFileName.IsEmpty())
+	if (World != nullptr && !Settings->DumpDependencyGraphFileName.IsEmpty())
 	{
 		DependencyGraphFileName = FString::Printf(TEXT("%s_%s"), *Settings->DumpDependencyGraphFileName, *ToString(World->GetNetMode()));
 	}
diff --git a/Engine/Source/Editor/MassEntityEditor/Private/MassEntityEditorSubsystem.cpp b/Engine/Source/Editor/MassEntityEditor/Private/MassEntityEditorSubsystem.cpp
--- a/Engine/Source/Editor/MassEntityEditor/Private/MassEntityEditorSubsystem.cpp
+++ b/Engine/Source/Editor/MassEntityEditor/Private/MassEntityEditorSubsystem.cpp
@@ -90,6 +90,7 @@ UMassEntityEditorSubsystem::~UMassEntityEditorSubsystem()
 void UMassEntityEditorSubsystem::Initialize(FSubsystemCollectionBase& Collection)
 {
 	UMassEntitySettings* Settings = GetMutableDefault<UMassEntitySettings>();
+	check(Settings);
 	if (Settings->IsInitialized2())
 	{
 		InitializeMassInternals();
@@ -127,6 +128,7 @@ void UMassEntityEditorSubsystem::InitializeMassInternals()
 
 	// set up ProcessingPhasesConfig
 	UMassEntitySettings* Settings = GetMutableDefault<UMassEntitySettings>();
+	check(Settings);
 	TConstArrayView<FMassProcessingPhaseConfig> MainPhasesConfig = Settings->GetProcessingPhasesConfig();
 	for (int PhaseIndex = 0; PhaseIndex < (int)EMassProcessingPhase::MAX; ++PhaseIndex)
 	{
@@ -151,6 +153,7 @@ void UMassEntityEditorSubsystem::StopAndCleanUp()
 	if (OnSettingsInitializedHandle.IsValid())
 	{
 		UMassEntitySettings* Settings = GetMutableDefault<UMassEntitySettings>();
+		check(Settings);
 		Settings->GetOnInitialized().Remove(OnSettingsInitializedHandle);
 	}
 }
diff --git a/Engine/Source/Runtime/MassEntity/Private/MassProcessor.cpp b/Engine/Source/Runtime/MassEntity/Private/MassProcessor.cpp
--- a/Engine/Source/Runtime/MassEntity/Private/MassProcessor.cpp
+++ b/Engine/Source/Runtime/MassEntity/Private/MassProcessor.cpp
@@ -605,6 +605,7 @@ void UMassCompositeProcessor::UpdateProcessorsCollection(TArrayView<FMassProcess
 
 	const UWorld* World = GetWorld();
 	const UMassEntitySettings* Settings = GetDefault<UMassEntitySettings>();
+	check(Settings);
 
 	const EProcessorExecutionFlags WorldExecutionFlags = UE::Mass::Utils::DetermineProcessorExecutionFlags(World, InWorldExecutionFlags);
 	const FMassProcessingPhaseConfig& PhaseConfig = Settings->GetProcessingPhaseConfig(ProcessingPhase);
diff --git a/Engine/Source/Runtime/MassEntity/Public/MassEntitySettings.h b/Engine/Source/Runtime/MassEntity/Public/MassEntitySettings.h
--- a/Engine/Source/Runtime/MassEntity/Public/MassEntitySettings.h
+++ b/Engine/Source/Runtime/MassEntity/Public/MassEntitySettings.h
@@ -11,7 +11,7 @@
 #include "MassEntitySettings.generated.h"
 
 #define GET_MASS_CONFIG_VALUE(a) \
-	UE_DEPRECATED_MACRO(5.8, "Use GetDefault<UMassEntitySettings>(*) directly.") \
+	UE_DEPRECATED_MACRO(5.8, "Use GetDefault<UMassEntitySettings>() (or GetMutableDefault) directly.") \
 	(GetMutableDefault<UMassEntitySettings>()->a)
 
 /**
@@ -60,8 +60,8 @@ public:
 
 protected:
 #if WITH_EDITOR
-	MASSENTITY_API virtual void PostEditChangeProperty(struct FPropertyChangedEvent& PropertyChangedEvent) override;
-	MASSENTITY_API virtual void PostEditChangeChainProperty(struct FPropertyChangedChainEvent& PropertyChangedEvent) override;
+	MASSENTITY_API virtual void PostEditChangeProperty(FPropertyChangedEvent& PropertyChangedEvent) override;
+	MASSENTITY_API virtual void PostEditChangeChainProperty(FPropertyChangedChainEvent& PropertyChangedEvent) override;
 #endif // WITH_EDITOR
 
 	MASSENTITY_API virtual void PostInitProperties() override;
