Mass/TEDS: Changes to support enabling TEDS in targets that run the editor with -game.

Mass: Change MassEntityEditorSubsystem to be a UEngineSubsystem selectively enabled when WITH_EDITOR is true
TEDS: Change TedsCore to be an "UncookedOnly" module

#Synced-CL 49950745

diff --git a/Engine/Plugins/Experimental/EditorDataStorage/EditorDataStorage.uplugin b/Engine/Plugins/Experimental/EditorDataStorage/EditorDataStorage.uplugin
--- a/Engine/Plugins/Experimental/EditorDataStorage/EditorDataStorage.uplugin
+++ b/Engine/Plugins/Experimental/EditorDataStorage/EditorDataStorage.uplugin
@@ -18,7 +18,7 @@
 	"Modules": [
 		{
 			"Name": "TedsCore",
-			"Type": "EditorAndProgram"
+			"Type": "UncookedOnly"
 		},
 		{
 			"Name": "TedsUI",
diff --git a/Engine/Plugins/Experimental/EditorDataStorage/Source/TedsCore/Private/TypedElementDatabase.cpp b/Engine/Plugins/Experimental/EditorDataStorage/Source/TedsCore/Private/TypedElementDatabase.cpp
--- a/Engine/Plugins/Experimental/EditorDataStorage/Source/TedsCore/Private/TypedElementDatabase.cpp
+++ b/Engine/Plugins/Experimental/EditorDataStorage/Source/TedsCore/Private/TypedElementDatabase.cpp
@@ -144,8 +144,7 @@ void UEditorDataStorage::Initialize()
 {
 	using namespace UE::Editor::DataStorage;
 
-	check(GEditor);
-	UMassEntityEditorSubsystem* Mass = GEditor->GetEditorSubsystem<UMassEntityEditorSubsystem>();
+	UMassEntityEditorSubsystem* Mass = UMassEntityEditorSubsystem::Get();
 	check(Mass);
 	OnPreMassTickHandle = Mass->GetOnPreTickDelegate().AddUObject(this, &UEditorDataStorage::OnPreMassTick);
 	OnPostMassTickHandle = Mass->GetOnPostTickDelegate().AddUObject(this, &UEditorDataStorage::OnPostMassTick);
@@ -1648,7 +1647,7 @@ void UEditorDataStorage::FinalizePhase(UE::Editor::DataStorage::EQueryTickPhase
 
 void UEditorDataStorage::Reset()
 {
-	if (UMassEntityEditorSubsystem* Mass = GEditor->GetEditorSubsystem<UMassEntityEditorSubsystem>())
+	if (UMassEntityEditorSubsystem* Mass = UMassEntityEditorSubsystem::Get())
 	{
 		Mass->GetOnPostTickDelegate().Remove(OnPostMassTickHandle);
 		Mass->GetOnPreTickDelegate().Remove(OnPreMassTickHandle);
diff --git a/Engine/Plugins/Runtime/MassGameplay/Source/MassGameplayEditor/Private/MassActorEditorSubsystem.cpp b/Engine/Plugins/Runtime/MassGameplay/Source/MassGameplayEditor/Private/MassActorEditorSubsystem.cpp
--- a/Engine/Plugins/Runtime/MassGameplay/Source/MassGameplayEditor/Private/MassActorEditorSubsystem.cpp
+++ b/Engine/Plugins/Runtime/MassGameplay/Source/MassGameplayEditor/Private/MassActorEditorSubsystem.cpp
@@ -14,7 +14,7 @@
 #include UE_INLINE_GENERATED_CPP_BY_NAME(MassActorEditorSubsystem)
 void UMassActorEditorSubsystem::Initialize(FSubsystemCollectionBase& Collection)
 {
-	UMassEntityEditorSubsystem* MassEditorEditorSubsystem = Collection.InitializeDependency<UMassEntityEditorSubsystem>();
+	UMassEntityEditorSubsystem* MassEditorEditorSubsystem = UMassEntityEditorSubsystem::Get();
 	check(MassEditorEditorSubsystem);
 	TSharedRef<FMassEntityManager> MassEntityManager = MassEditorEditorSubsystem->GetMutableEntityManager();
 	ActorManager = MakeShareable(new FMassActorManager(MassEntityManager));
diff --git a/Engine/Source/Editor/MassEntityEditor/Private/MassEntityEditorSubsystem.cpp b/Engine/Source/Editor/MassEntityEditor/Private/MassEntityEditorSubsystem.cpp
--- a/Engine/Source/Editor/MassEntityEditor/Private/MassEntityEditorSubsystem.cpp
+++ b/Engine/Source/Editor/MassEntityEditor/Private/MassEntityEditorSubsystem.cpp
@@ -4,6 +4,7 @@
 #include "MassEntityManager.h"
 #include "MassEntitySettings.h"
 #include "TickableEditorObject.h"
+#include "Engine/Engine.h"
 
 #include UE_INLINE_GENERATED_CPP_BY_NAME(MassEntityEditorSubsystem)
 
@@ -87,6 +88,11 @@ UMassEntityEditorSubsystem::~UMassEntityEditorSubsystem()
 	}
 }
 
+UMassEntityEditorSubsystem* UMassEntityEditorSubsystem::Get()
+{
+	return GEngine ? GEngine->GetEngineSubsystem<UMassEntityEditorSubsystem>() : nullptr;
+}
+
 void UMassEntityEditorSubsystem::Initialize(FSubsystemCollectionBase& Collection)
 {
 	UMassEntitySettings* Settings = GetMutableDefault<UMassEntitySettings>();
@@ -109,6 +115,15 @@ void UMassEntityEditorSubsystem::Deinitialize()
 	Super::Deinitialize();
 }
 
+bool UMassEntityEditorSubsystem::ShouldCreateSubsystem(UObject* Outer) const
+{
+#if WITH_EDITOR
+	return true;
+#else
+	return false;
+#endif
+}
+
 void UMassEntityEditorSubsystem::InitializeMassInternals()
 {
 	// Note that since this is a raw pointer we need to remember to delete it once we're done. It's now done in both 
diff --git a/Engine/Source/Editor/MassEntityEditor/Public/MassEntityEditorSubsystem.h b/Engine/Source/Editor/MassEntityEditor/Public/MassEntityEditorSubsystem.h
--- a/Engine/Source/Editor/MassEntityEditor/Public/MassEntityEditorSubsystem.h
+++ b/Engine/Source/Editor/MassEntityEditor/Public/MassEntityEditorSubsystem.h
@@ -2,8 +2,8 @@
 
 #pragma once
 
-#include "EditorSubsystem.h"
 #include "MassProcessingPhaseManager.h"
+#include "Subsystems/EngineSubsystem.h"
 #include "MassEntityEditorSubsystem.generated.h"
 
 #define UE_API MASSENTITYEDITOR_API
@@ -15,8 +15,11 @@ namespace UE::Mass
 	struct FMassEditorTickFunction;
 }
 
+/*
+ * This is an Engine Subsystem that is selectively enabled using WITH_EDITOR so it is available in custom programs that run the editor with -game
+ */
 UCLASS(MinimalAPI)
-class UMassEntityEditorSubsystem : public UEditorSubsystem
+class UMassEntityEditorSubsystem : public UEngineSubsystem
 {
 	GENERATED_BODY()
 
@@ -27,6 +30,8 @@ public:
 	UE_API UMassEntityEditorSubsystem();
 	UE_API ~UMassEntityEditorSubsystem();
 
+	UE_API static UMassEntityEditorSubsystem* Get();
+
 	TSharedRef<FMassEntityManager> GetMutableEntityManager() { return EntityManager; }
 	TSharedRef<FMassProcessingPhaseManager> GetMutablePhaseManager() { return PhaseManager; }
 
@@ -42,6 +47,7 @@ protected:
 	// USubsystem interface begin
 	UE_API virtual void Initialize(FSubsystemCollectionBase& Collection) override;
 	UE_API virtual void Deinitialize() override;
+	UE_API virtual bool ShouldCreateSubsystem(UObject* Outer) const override;
 	// USubsystem interface end
 
 	UE_API void InitializeMassInternals();
