[MassDebug]
Add Rewind Debugger integration for Mass entities

Introduces a MassChannel trace channel toggle, an entity event cache, and
per-entity RewindDebugger tracks with timeline visualization and a diff-style
details panel for inspecting fragment composition changes over time.

- Move SEventTimelineView to RewindDebuggerInterface module for shared access
- Add TRACE_OBJECT/TRACE_INSTANCE calls to MassEntityComponent entity creation
- Add FMassRewindDebuggerRuntimeExtension to toggle MassChannel on recording
- Add FEntityEventCache for indexed per-entity/per-archetype event lookups
- Add FMassEntityTrack (per-entity leaf) RewindDebugger track
- Add SMassEntityDetailsView with prev/current archetype fragment diff table

#Synced-CL 52361227

diff --git a/Engine/Plugins/Animation/GameplayInsights/Source/GameplayInsights/Private/NotifiesTrack.cpp b/Engine/Plugins/Animation/GameplayInsights/Source/GameplayInsights/Private/NotifiesTrack.cpp
--- a/Engine/Plugins/Animation/GameplayInsights/Source/GameplayInsights/Private/NotifiesTrack.cpp
+++ b/Engine/Plugins/Animation/GameplayInsights/Source/GameplayInsights/Private/NotifiesTrack.cpp
@@ -80,9 +80,9 @@ bool FNotifyTrack::UpdateInternal()
 		TRACE_CPUPROFILER_EVENT_SCOPE(FNotifiesTrack::UpdateEventPointsInternal);
 		EventUpdateRequested = 0;
 		
-		auto& EventPoints = EventData->Points;
+		auto& EventPoints = EventData->EventPoints;
 		EventPoints.SetNum(0,EAllowShrinking::No);
-		EventData->Windows.SetNum(0);
+		EventData->EventWindows.SetNum(0);
 
 		TraceServices::FProviderReadScopeLock AnimationProviderReadScope(*AnimationProvider);
 
@@ -94,7 +94,7 @@ bool FNotifyTrack::UpdateInternal()
 				{
 					if (InMessage.NameId == NotifyTrackId.NameId && MakeTrackType(InMessage.NotifyEventType) == NotifyTrackId.Type)
 					{
-						SEventTimelineView::FTimelineEventData::EventPoint Point;
+						SEventTimelineView::FTimelineEventData::FEventPoint Point;
 						Point.Time = InMessage.RecordingTime;
 						Point.Color = MakeNotifyColor(InMessage.NameId);
 						EventPoints.Add(Point);
@@ -112,11 +112,11 @@ bool FNotifyTrack::UpdateInternal()
 				{
 					if (InMessage.NameId == NotifyTrackId.NameId && MakeTrackType(InMessage.NotifyEventType) == NotifyTrackId.Type)
 					{
-						SEventTimelineView::FTimelineEventData::EventWindow Window;
+						SEventTimelineView::FTimelineEventData::FEventWindow Window;
 						Window.TimeStart = InMessage.RecordingTime;
 						Window.TimeEnd = InMessage.RecordingTime + InMessage.Duration;
 						Window.Color = MakeNotifyColor(InMessage.NameId);
-						EventData->Windows.Add(Window);
+						EventData->EventWindows.Add(Window);
 					}
 				}
 				
diff --git a/Engine/Plugins/Animation/GameplayInsights/Source/GameplayInsights/Private/SEventTimelineView.cpp b/Engine/Plugins/Animation/GameplayInsights/Source/GameplayInsights/Private/SEventTimelineView.cpp
deleted file mode 100644
diff --git a/Engine/Plugins/Animation/GameplayInsights/Source/GameplayInsights/Public/SEventTimelineView.h b/Engine/Plugins/Animation/GameplayInsights/Source/GameplayInsights/Public/SEventTimelineView.h
deleted file mode 100644
diff --git a/Engine/Plugins/Animation/GameplayInsights/Source/RewindDebugger/Private/RewindDebuggerObjectTrack.cpp b/Engine/Plugins/Animation/GameplayInsights/Source/RewindDebugger/Private/RewindDebuggerObjectTrack.cpp
--- a/Engine/Plugins/Animation/GameplayInsights/Source/RewindDebugger/Private/RewindDebuggerObjectTrack.cpp
+++ b/Engine/Plugins/Animation/GameplayInsights/Source/RewindDebugger/Private/RewindDebuggerObjectTrack.cpp
@@ -299,7 +299,7 @@ bool FRewindDebuggerObjectTrack::UpdateInternal()
 
 		const TRange<double> Existence = GameplayProvider->GetObjectRecordingLifetime(ObjectId);
 
-		ExistenceRange->Windows.SetNum(0, EAllowShrinking::No);
+		ExistenceRange->EventWindows.SetNum(0, EAllowShrinking::No);
 		if (Existence.HasLowerBound() && Existence.HasUpperBound())
 		{
 			double EndTime = Existence.GetUpperBoundValue();
@@ -307,7 +307,7 @@ bool FRewindDebuggerObjectTrack::UpdateInternal()
 			{
 				EndTime = GameplayProvider->GetRecordingDuration();
 			}
-			ExistenceRange->Windows.Add({ Existence.GetLowerBoundValue(), EndTime, LOCTEXT("Object Existence","Object Existence"), LOCTEXT("Object Existence","Object Existence"), FLinearColor(0.1f,0.11f,0.1f) });
+			ExistenceRange->EventWindows.Add({ Existence.GetLowerBoundValue(), EndTime, LOCTEXT("Object Existence","Object Existence"), LOCTEXT("Object Existence","Object Existence"), FLinearColor(0.1f,0.11f,0.1f) });
 		}
 
 		if (!bIconSearched)
diff --git a/Engine/Plugins/Animation/GameplayInsights/Source/RewindDebuggerVLog/Private/VisualLogTrack.cpp b/Engine/Plugins/Animation/GameplayInsights/Source/RewindDebuggerVLog/Private/VisualLogTrack.cpp
--- a/Engine/Plugins/Animation/GameplayInsights/Source/RewindDebuggerVLog/Private/VisualLogTrack.cpp
+++ b/Engine/Plugins/Animation/GameplayInsights/Source/RewindDebuggerVLog/Private/VisualLogTrack.cpp
@@ -100,8 +100,8 @@ bool FVisualLogCategoryTrack::UpdateInternal()
 			TRACE_CPUPROFILER_EVENT_SCOPE(FVisualLogTrack::UpdateEventPointsInternal);
 			EventUpdateRequested = 0;
 
-			EventData->Points.SetNum(0, EAllowShrinking::No);
-			EventData->Windows.SetNum(0);
+			EventData->EventPoints.SetNum(0, EAllowShrinking::No);
+			EventData->EventWindows.SetNum(0);
 
 			TraceServices::FProviderReadScopeLock ProviderReadScope(*VisLogProvider);
 
@@ -113,7 +113,7 @@ bool FVisualLogCategoryTrack::UpdateInternal()
 							{
 								if (Element.Category == Category)
 								{
-									EventData->Points.Add({ InMessage.TimeStamp,GetCategoryNameAsText(Element.Category), FText::FromString(Element.Description),Element.GetFColor() });
+									EventData->EventPoints.Add({ InMessage.TimeStamp,GetCategoryNameAsText(Element.Category), FText::FromString(Element.Description),Element.GetFColor() });
 								}
 							}
 							for (const FVisualLogLine& Line : InMessage.LogLines)
@@ -121,7 +121,7 @@ bool FVisualLogCategoryTrack::UpdateInternal()
 								if (Line.Category == Category)
 								{
 									// No Description from log lines.
-									EventData->Points.Add({ InMessage.TimeStamp, GetCategoryNameAsText(Line.Category), FText::GetEmpty(), Line.Color });
+									EventData->EventPoints.Add({ InMessage.TimeStamp, GetCategoryNameAsText(Line.Category), FText::GetEmpty(), Line.Color });
 								}
 							}
 							return TraceServices::EEventEnumerate::Continue;
diff --git a/Engine/Plugins/Chooser/Source/ChooserEditor/Private/ChooserTrack.cpp b/Engine/Plugins/Chooser/Source/ChooserEditor/Private/ChooserTrack.cpp
--- a/Engine/Plugins/Chooser/Source/ChooserEditor/Private/ChooserTrack.cpp
+++ b/Engine/Plugins/Chooser/Source/ChooserEditor/Private/ChooserTrack.cpp
@@ -82,8 +82,8 @@ bool FChooserTrack::UpdateInternal()
 		TRACE_CPUPROFILER_EVENT_SCOPE(FChooserTrack::UpdateEventPointsInternal);
 		EventUpdateRequested = 0;
 		
-		EventData->Points.SetNum(0,EAllowShrinking::No);
-		EventData->Windows.SetNum(0);
+		EventData->EventPoints.SetNum(0,EAllowShrinking::No);
+		EventData->EventWindows.SetNum(0);
 
 		TraceServices::FProviderReadScopeLock ProviderReadScope(*ChooserProvider);
 		ChooserProvider->ReadChooserEvaluationTimeline(ObjectId, [this, StartTime, EndTime](const FChooserProvider::ChooserEvaluationTimeline& InTimeline)
@@ -92,7 +92,7 @@ bool FChooserTrack::UpdateInternal()
 			{
 				if (ChooserEvaluationData.ChooserId == ChooserId)
 				{
-					EventData->Points.Add({InStartTime, FText(), FText(), FLinearColor::White});
+					EventData->EventPoints.Add({InStartTime, FText(), FText(), FLinearColor::White});
 				}
 				return TraceServices::EEventEnumerate::Continue;
 			});
diff --git a/Engine/Plugins/Developer/FieldNotificationTrace/Source/FieldNotificationTraceEditor/Private/FieldNotificationTrack.cpp b/Engine/Plugins/Developer/FieldNotificationTrace/Source/FieldNotificationTraceEditor/Private/FieldNotificationTrack.cpp
--- a/Engine/Plugins/Developer/FieldNotificationTrace/Source/FieldNotificationTraceEditor/Private/FieldNotificationTrack.cpp
+++ b/Engine/Plugins/Developer/FieldNotificationTrace/Source/FieldNotificationTraceEditor/Private/FieldNotificationTrack.cpp
@@ -65,8 +65,8 @@ bool FFieldNotifyTrack::UpdateInternal()
 		TRACE_CPUPROFILER_EVENT_SCOPE(FFieldNotifyTrack::UpdateEventPointsInternal);
 		EventUpdateRequested = 0;
 		
-		EventData->Points.SetNum(0,EAllowShrinking::No);
-		EventData->Windows.SetNum(0);
+		EventData->EventPoints.SetNum(0,EAllowShrinking::No);
+		EventData->EventWindows.SetNum(0);
 
 		TraceServices::FAnalysisSessionReadScope SessionReadScope(*AnalysisSession);
 		
@@ -74,7 +74,7 @@ bool FFieldNotifyTrack::UpdateInternal()
 			{
 				if (EvaluationData.FieldNotifyId == Self->FieldNotifyId)
 				{
-					Self->EventData->Points.Add({ InStartTime, FText(), FText(), FLinearColor::White });
+					Self->EventData->EventPoints.Add({ InStartTime, FText(), FText(), FLinearColor::White });
 				}
 			});
 	}
diff --git a/Engine/Plugins/Experimental/UAF/UAF/Source/UAFEditor/Private/RewindDebugger/AnimNextModuleTrack.cpp b/Engine/Plugins/Experimental/UAF/UAF/Source/UAFEditor/Private/RewindDebugger/AnimNextModuleTrack.cpp
--- a/Engine/Plugins/Experimental/UAF/UAF/Source/UAFEditor/Private/RewindDebugger/AnimNextModuleTrack.cpp
+++ b/Engine/Plugins/Experimental/UAF/UAF/Source/UAFEditor/Private/RewindDebugger/AnimNextModuleTrack.cpp
@@ -82,7 +82,7 @@ FAnimNextModuleTrack::~FAnimNextModuleTrack()
 void FAnimNextModuleTrack::Initialize()
 {
 	ExistenceRange = MakeShared<SEventTimelineView::FTimelineEventData>();
-	ExistenceRange->Windows.Add({0,0, GetDisplayNameInternal(), GetDisplayNameInternal(), FLinearColor(0.1f,0.15f,0.11f)});
+	ExistenceRange->EventWindows.Add({0,0, GetDisplayNameInternal(), GetDisplayNameInternal(), FLinearColor(0.1f,0.15f,0.11f)});
 
 	FPropertyEditorModule& PropertyEditorModule = FModuleManager::GetModuleChecked<FPropertyEditorModule>("PropertyEditor");
 	FDetailsViewArgs DetailsViewArgs;
@@ -205,9 +205,9 @@ bool FAnimNextModuleTrack::UpdateInternal()
 					}
 				}
 
-				ExistenceRange->Windows.SetNum(1, EAllowShrinking::No);
-				ExistenceRange->Windows[0].TimeStart = Data->StartTime;
-				ExistenceRange->Windows[0].TimeEnd = Data->EndTime;
+				ExistenceRange->EventWindows.SetNum(1, EAllowShrinking::No);
+				ExistenceRange->EventWindows[0].TimeStart = Data->StartTime;
+				ExistenceRange->EventWindows[0].TimeEnd = Data->EndTime;
 			}
 		}
 	}
diff --git a/Engine/Plugins/Experimental/UAF/UAFAnimGraph/Source/UAFAnimGraphEditor/Private/RewindDebugger/EvaluationProgramTrack.cpp b/Engine/Plugins/Experimental/UAF/UAFAnimGraph/Source/UAFAnimGraphEditor/Private/RewindDebugger/EvaluationProgramTrack.cpp
--- a/Engine/Plugins/Experimental/UAF/UAFAnimGraph/Source/UAFAnimGraphEditor/Private/RewindDebugger/EvaluationProgramTrack.cpp
+++ b/Engine/Plugins/Experimental/UAF/UAFAnimGraph/Source/UAFAnimGraphEditor/Private/RewindDebugger/EvaluationProgramTrack.cpp
@@ -72,7 +72,7 @@ FEvaluationProgramTrack::~FEvaluationProgramTrack()
 void FEvaluationProgramTrack::Initialize()
 {
 	ExistenceRange = MakeShared<SEventTimelineView::FTimelineEventData>();
-	ExistenceRange->Windows.Add({0,0, GetDisplayNameInternal(), GetDisplayNameInternal(), FLinearColor(0.1f,0.15f,0.11f)});
+	ExistenceRange->EventWindows.Add({0,0, GetDisplayNameInternal(), GetDisplayNameInternal(), FLinearColor(0.1f,0.15f,0.11f)});
 	
    	FPropertyEditorModule& PropertyEditorModule = FModuleManager::GetModuleChecked<FPropertyEditorModule>("PropertyEditor");
    	FDetailsViewArgs DetailsViewArgs;
diff --git a/Engine/Plugins/Experimental/UAF/UAFAnimGraph/Source/UAFAnimGraphEditor/Private/RewindDebugger/SequenceInfoTrack.cpp b/Engine/Plugins/Experimental/UAF/UAFAnimGraph/Source/UAFAnimGraphEditor/Private/RewindDebugger/SequenceInfoTrack.cpp
--- a/Engine/Plugins/Experimental/UAF/UAFAnimGraph/Source/UAFAnimGraphEditor/Private/RewindDebugger/SequenceInfoTrack.cpp
+++ b/Engine/Plugins/Experimental/UAF/UAFAnimGraph/Source/UAFAnimGraphEditor/Private/RewindDebugger/SequenceInfoTrack.cpp
@@ -229,7 +229,7 @@ FSequenceInfoTrack::~FSequenceInfoTrack()
 void FSequenceInfoTrack::Initialize()
 {
 	ExistenceRange = MakeShared<SEventTimelineView::FTimelineEventData>();
-	ExistenceRange->Windows.Add({0, 0, GetDisplayNameInternal(), GetDisplayNameInternal(), FLinearColor(0.1f, 0.15f, 0.11f)});
+	ExistenceRange->EventWindows.Add({0, 0, GetDisplayNameInternal(), GetDisplayNameInternal(), FLinearColor(0.1f, 0.15f, 0.11f)});
 
 	FPropertyEditorModule& PropertyEditorModule = FModuleManager::GetModuleChecked<FPropertyEditorModule>("PropertyEditor");
 	FDetailsViewArgs DetailsViewArgs;
diff --git a/Engine/Plugins/Experimental/UAF/UAFAnimNode/Source/UAFAnimNodeEditor/Private/RewindDebugger/UAFAnimOpTrack.cpp b/Engine/Plugins/Experimental/UAF/UAFAnimNode/Source/UAFAnimNodeEditor/Private/RewindDebugger/UAFAnimOpTrack.cpp
--- a/Engine/Plugins/Experimental/UAF/UAFAnimNode/Source/UAFAnimNodeEditor/Private/RewindDebugger/UAFAnimOpTrack.cpp
+++ b/Engine/Plugins/Experimental/UAF/UAFAnimNode/Source/UAFAnimNodeEditor/Private/RewindDebugger/UAFAnimOpTrack.cpp
@@ -71,7 +71,7 @@ FUAFAnimOpTrack::~FUAFAnimOpTrack()
 void FUAFAnimOpTrack::Initialize()
 {
 	ExistenceRange = MakeShared<SEventTimelineView::FTimelineEventData>();
-	ExistenceRange->Windows.Add({0,0, GetDisplayNameInternal(), GetDisplayNameInternal(), FLinearColor(0.1f,0.15f,0.11f)});
+	ExistenceRange->EventWindows.Add({0,0, GetDisplayNameInternal(), GetDisplayNameInternal(), FLinearColor(0.1f,0.15f,0.11f)});
 	
    	FPropertyEditorModule& PropertyEditorModule = FModuleManager::GetModuleChecked<FPropertyEditorModule>("PropertyEditor");
    	FDetailsViewArgs DetailsViewArgs;
diff --git a/Engine/Plugins/Experimental/UAF/UAFAnimNode/Source/UAFAnimNodeEditor/Private/RewindDebugger/UAFSequenceInfoTrack.cpp b/Engine/Plugins/Experimental/UAF/UAFAnimNode/Source/UAFAnimNodeEditor/Private/RewindDebugger/UAFSequenceInfoTrack.cpp
--- a/Engine/Plugins/Experimental/UAF/UAFAnimNode/Source/UAFAnimNodeEditor/Private/RewindDebugger/UAFSequenceInfoTrack.cpp
+++ b/Engine/Plugins/Experimental/UAF/UAFAnimNode/Source/UAFAnimNodeEditor/Private/RewindDebugger/UAFSequenceInfoTrack.cpp
@@ -229,7 +229,7 @@ FUAFSequenceInfoTrack::~FUAFSequenceInfoTrack()
 void FUAFSequenceInfoTrack::Initialize()
 {
 	ExistenceRange = MakeShared<SEventTimelineView::FTimelineEventData>();
-	ExistenceRange->Windows.Add({0, 0, GetDisplayNameInternal(), GetDisplayNameInternal(), FLinearColor(0.1f, 0.15f, 0.11f)});
+	ExistenceRange->EventWindows.Add({0, 0, GetDisplayNameInternal(), GetDisplayNameInternal(), FLinearColor(0.1f, 0.15f, 0.11f)});
 
 	FPropertyEditorModule& PropertyEditorModule = FModuleManager::GetModuleChecked<FPropertyEditorModule>("PropertyEditor");
 	FDetailsViewArgs DetailsViewArgs;
diff --git a/Engine/Plugins/Runtime/RigVM/Source/RigVMEditor/Private/Trace/RigVMTraceTrack.cpp b/Engine/Plugins/Runtime/RigVM/Source/RigVMEditor/Private/Trace/RigVMTraceTrack.cpp
--- a/Engine/Plugins/Runtime/RigVM/Source/RigVMEditor/Private/Trace/RigVMTraceTrack.cpp
+++ b/Engine/Plugins/Runtime/RigVM/Source/RigVMEditor/Private/Trace/RigVMTraceTrack.cpp
@@ -104,12 +104,12 @@ bool FRewindDebuggerRigVMExecuteTrack::UpdateInternal()
 	{
 		EventUpdateRequested = 0;
 	
-		EventData->Points.SetNum(0,EAllowShrinking::No);
-		EventData->Windows.SetNum(1,EAllowShrinking::No);
+		EventData->EventPoints.SetNum(0,EAllowShrinking::No);
+		EventData->EventWindows.SetNum(1,EAllowShrinking::No);
 
-		EventData->Windows[0].Color = FLinearColor::White;
-		EventData->Windows[0].TimeStart = DBL_MAX;
-		EventData->Windows[0].TimeEnd = DBL_MIN;
+		EventData->EventWindows[0].Color = FLinearColor::White;
+		EventData->EventWindows[0].TimeStart = DBL_MAX;
+		EventData->EventWindows[0].TimeEnd = DBL_MIN;
 
 		TraceServices::FAnalysisSessionReadScope SessionReadScope(*AnalysisSession);
 	
@@ -120,9 +120,9 @@ bool FRewindDebuggerRigVMExecuteTrack::UpdateInternal()
 				if (ExecuteData.HostId == HostId)
 				{
 					const double StartRecordingTime = TraceToRecordingTime(InStartTime);
-					EventData->Windows[0].TimeStart = FMath::Min(StartRecordingTime, EventData->Windows[0].TimeStart);
-					EventData->Windows[0].TimeEnd = FMath::Max(StartRecordingTime, EventData->Windows[0].TimeEnd);
-					//EventData->Points.Add({TraceToRecordingTime(InStartTime), FText(), FText(), FLinearColor::White});
+					EventData->EventWindows[0].TimeStart = FMath::Min(StartRecordingTime, EventData->EventWindows[0].TimeStart);
+					EventData->EventWindows[0].TimeEnd = FMath::Max(StartRecordingTime, EventData->EventWindows[0].TimeEnd);
+					//EventData->EventPoints.Add({TraceToRecordingTime(InStartTime), FText(), FText(), FLinearColor::White});
 				}
 				return TraceServices::EEventEnumerate::Continue;
 			});
diff --git a/Engine/Source/Editor/MassEntityEditor/MassEntityEditor.Build.cs b/Engine/Source/Editor/MassEntityEditor/MassEntityEditor.Build.cs
--- a/Engine/Source/Editor/MassEntityEditor/MassEntityEditor.Build.cs
+++ b/Engine/Source/Editor/MassEntityEditor/MassEntityEditor.Build.cs
@@ -16,28 +16,33 @@ namespace UnrealBuildTool.Rules
 					"Core",
 					"CoreUObject",
 					"Engine",
-					"InputCore",
-					"AssetTools",
-					"UnrealEd",
-					"Slate",
-					"SlateCore",
-					"PropertyEditor",
-					"MassCore",
 					"MassEntity",
-					"DetailCustomizations",
-					"ComponentVisualizers",
-					"Projects",
-					"EditorSubsystem",
+					"SlateCore",
+					"UnrealEd",
 				}
 			);
 
 			PrivateDependencyModuleNames.AddRange(
 				new string[] {
-					"RenderCore",
+					"AIGraph",
+					"AssetTools",
+					"ComponentVisualizers",
+					"DetailCustomizations",
+					"EditorSubsystem",
 					"GraphEditor",
+					"InputCore",
 					"KismetWidgets",
-					"AIGraph",
+					"MassCore",
+					"MassDeveloper",
+					"Projects",
+					"PropertyEditor",
+					"RenderCore",
+					"RewindDebuggerInterface",
+					"RewindDebuggerRuntimeInterface",
+					"Slate",
 					"ToolMenus",
+					"TraceLog",
+					"TraceServices",
 				}
 			);
 
diff --git a/Engine/Source/Editor/MassEntityEditor/Private/MassEntityEditorModule.cpp b/Engine/Source/Editor/MassEntityEditor/Private/MassEntityEditorModule.cpp
--- a/Engine/Source/Editor/MassEntityEditor/Private/MassEntityEditorModule.cpp
+++ b/Engine/Source/Editor/MassEntityEditor/Private/MassEntityEditorModule.cpp
@@ -1,14 +1,19 @@
 // Copyright Epic Games, Inc. All Rights Reserved.
 
 #include "MassEntityEditorModule.h"
+
+#include "Features/IModularFeatures.h"
+#include "IRewindDebuggerTrackCreator.h"
 #include "MassEditorStyle.h"
 #include "Modules/ModuleManager.h"
+#include "Trace/MassEntityTrack.h"
+#include "Trace/MassRewindDebuggerIntegration.h"
 #if WITH_UNREAL_DEVELOPER_TOOLS
-#include "MessageLogModule.h"
 #include "Engine/World.h"
 #include "Logging/MessageLog.h"
 #include "MassDebugger.h"
 #include "MassEntityEditor.h"
+#include "MessageLogModule.h"
 #endif // WITH_UNREAL_DEVELOPER_TOOLS
 
 #define LOCTEXT_NAMESPACE "Mass"
@@ -17,6 +22,8 @@ IMPLEMENT_MODULE(FMassEntityEditorModule, MassEntityEditor)
 
 void FMassEntityEditorModule::StartupModule()
 {
+	using namespace UE::Mass;
+
 	MenuExtensibilityManager = MakeShareable(new FExtensibilityManager);
 	ToolBarExtensibilityManager = MakeShareable(new FExtensibilityManager);
 
@@ -27,17 +34,39 @@ void FMassEntityEditorModule::StartupModule()
 	FMessageLogInitializationOptions InitOptions;
 	InitOptions.bShowPages = true;
 	InitOptions.bShowFilters = true;
-	MessageLogModule.RegisterLogListing(UE::Mass::Editor::MessageLogPageName
-		, FText::FromName(UE::Mass::Editor::MessageLogPageName), InitOptions);
+	MessageLogModule.RegisterLogListing(Editor::MessageLogPageName
+		, FText::FromName(Editor::MessageLogPageName), InitOptions);
 
 	OnWorldCleanupHandle = FWorldDelegates::OnWorldCleanup.AddStatic(&FMassEntityEditorModule::OnWorldCleanup);
 
 	FModuleManager::Get().LoadModule("MassEntityDebugger");
 #endif // WITH_UNREAL_DEVELOPER_TOOLS
+
+	// RewindDebugger integration
+	RuntimeExtension = MakeUnique<FMassRewindDebuggerRuntimeExtension>();
+	EntityTrackCreator = MakeUnique<Trace::FEntityTrackCreator>();
+
+	IModularFeatures& ModularFeatures = IModularFeatures::Get();
+	ModularFeatures.RegisterModularFeature(IRewindDebuggerRuntimeExtension::ModularFeatureName, RuntimeExtension.Get());
+	ModularFeatures.RegisterModularFeature(RewindDebugger::IRewindDebuggerTrackCreator::ModularFeatureName, EntityTrackCreator.Get());
 }
 
 void FMassEntityEditorModule::ShutdownModule()
 {
+	IModularFeatures& ModularFeatures = IModularFeatures::Get();
+
+	// Unregister RewindDebugger integration
+	if (RuntimeExtension.IsValid())
+	{
+		ModularFeatures.UnregisterModularFeature(IRewindDebuggerRuntimeExtension::ModularFeatureName, RuntimeExtension.Get());
+	}
+	if (EntityTrackCreator.IsValid())
+	{
+		ModularFeatures.UnregisterModularFeature(RewindDebugger::IRewindDebuggerTrackCreator::ModularFeatureName, EntityTrackCreator.Get());
+	}
+	RuntimeExtension.Reset();
+	EntityTrackCreator.Reset();
+
 	ProcessorClassCache.Reset();
 	MenuExtensibilityManager.Reset();
 	ToolBarExtensibilityManager.Reset();
diff --git a/Engine/Source/Editor/MassEntityEditor/Private/Trace/MassEntityEventCache.cpp b/Engine/Source/Editor/MassEntityEditor/Private/Trace/MassEntityEventCache.cpp
new file mode 100644
--- /dev/null
+++ b/Engine/Source/Editor/MassEntityEditor/Private/Trace/MassEntityEventCache.cpp
@@ -0,0 +1,116 @@
+// Copyright Epic Games, Inc. All Rights Reserved.
+
+#include "MassEntityEventCache.h"
+
+#include "Common/ProviderLock.h"
+#include "Misc/ScopeRWLock.h"
+#include "TraceServices/Model/AnalysisSession.h"
+#include "TraceServices/Model/Frames.h"
+
+namespace UE::Mass::Trace
+{
+
+void FEntityEventCache::Update(
+	const IMassTraceProvider& Provider,
+	const TraceServices::IAnalysisSession& Session)
+{
+	{
+		TraceServices::FProviderReadScopeLock ReadLock(Provider);
+		const uint64 CurrentCounter = Provider.GetUpdateCounter();
+		if (CurrentCounter == LastUpdateCounter)
+		{
+			return;
+		}
+		LastUpdateCounter = CurrentCounter;
+
+		FWriteScopeLock WriteLock(CacheLock);
+		EntityEvents.Reset();
+		ArchetypeEntities.Reset();
+		ActiveArchetypeIds.Reset();
+
+		const uint64 TotalEvents = Provider.GetEntityEventCount();
+		if (TotalEvents == 0)
+		{
+			return;
+		}
+
+		Provider.EnumerateEntityEvents(0, TotalEvents,
+			[this](const FEntityEventRecord& Record, uint64 /*EventIndex*/)
+			{
+				FCachedEntityEventRecord& CachedRecord = EntityEvents.FindOrAdd(Record.Entity).AddDefaulted_GetRef();
+				static_cast<FEntityEventRecord&>(CachedRecord) = Record;
+				ArchetypeEntities.FindOrAdd(Record.ArchetypeID).Add(Record.Entity);
+				ActiveArchetypeIds.Add(Record.ArchetypeID);
+			});
+	}
+
+	ConvertTimesToRecordingTime(Session);
+
+	// Sort each entity's events by time so consumers can rely on ascending order
+	// (e.g. early-exit in ResolveArchetypesAtScrubTime).
+	{
+		FWriteScopeLock WriteLock(CacheLock);
+		for (TPair<uint64, TArray<FCachedEntityEventRecord>>& Pair : EntityEvents)
+		{
+			Pair.Value.Sort([](const FCachedEntityEventRecord& A, const FCachedEntityEventRecord& B)
+			{
+				return A.Time < B.Time;
+			});
+		}
+	}
+}
+
+void FEntityEventCache::ConvertTimesToRecordingTime(const TraceServices::IAnalysisSession& Session)
+{
+	TraceServices::FAnalysisSessionReadScope SessionReadScope(Session);
+	const TraceServices::IFrameProvider& FrameProvider = TraceServices::ReadFrameProvider(Session);
+
+	// Pass 1: resolve frame indices while Event.Time is still in ProfileTime space.
+	for (TPair<uint64, TArray<FCachedEntityEventRecord>>& Pair : EntityEvents)
+	{
+		for (FCachedEntityEventRecord& Event : Pair.Value)
+		{
+			Event.FrameIndex = FrameProvider.GetFrameNumberForTimestamp(TraceFrameType_Game, Event.Time);
+		}
+	}
+
+	// Pass 2: normalize event times relative to the first frame's start time.
+	const uint64 FrameCount = FrameProvider.GetFrameCount(TraceFrameType_Game);
+	if (FrameCount == 0)
+	{
+		return;
+	}
+
+	const TraceServices::FFrame* FirstFrame = FrameProvider.GetFrame(TraceFrameType_Game, 0);
+	if (!FirstFrame)
+	{
+		return;
+	}
+
+	const double ProfileStart = FirstFrame->StartTime;
+	for (TPair<uint64, TArray<FCachedEntityEventRecord>>& Pair : EntityEvents)
+	{
+		for (FCachedEntityEventRecord& Event : Pair.Value)
+		{
+			Event.Time -= ProfileStart;
+		}
+	}
+}
+
+TConstArrayView<FEntityEventCache::FCachedEntityEventRecord> FEntityEventCache::GetEntityEvents(const uint64 EntityId) const
+{
+	FReadScopeLock ReadLock(CacheLock);
+	if (const TArray<FCachedEntityEventRecord>* Events = EntityEvents.Find(EntityId))
+	{
+		return MakeArrayView(*Events);
+	}
+	return {};
+}
+
+const TSet<uint64>* FEntityEventCache::GetEntitiesForArchetype(const uint64 ArchetypeId) const
+{
+	FReadScopeLock ReadLock(CacheLock);
+	return ArchetypeEntities.Find(ArchetypeId);
+}
+
+} // namespace UE::Mass::Trace
diff --git a/Engine/Source/Editor/MassEntityEditor/Private/Trace/MassEntityEventCache.h b/Engine/Source/Editor/MassEntityEditor/Private/Trace/MassEntityEventCache.h
new file mode 100644
--- /dev/null
+++ b/Engine/Source/Editor/MassEntityEditor/Private/Trace/MassEntityEventCache.h
@@ -0,0 +1,76 @@
+// Copyright Epic Games, Inc. All Rights Reserved.
+
+#pragma once
+
+#include "Containers/Map.h"
+#include "Misc/ScopeRWLock.h"
+#include "Trace/MassTraceTypes.h"
+
+namespace TraceServices
+{
+class IAnalysisSession;
+}
+
+namespace UE::Mass::Trace
+{
+class IMassTraceProvider;
+/**
+ * Caches entity events indexed by entity and archetype for efficient sub-track querying.
+ * The IMassTraceProvider only offers sequential enumeration, so this cache builds
+ * per-entity and per-archetype indices to avoid O(N) scans in every sub-track update.
+ */
+class FEntityEventCache
+{
+public:
+	/** Extended event record with frame index resolved at cache build time. */
+	struct FCachedEntityEventRecord : FEntityEventRecord
+	{
+		/** Game frame index this event belongs to. Events in the same frame share the same index. */
+		uint32 FrameIndex = 0;
+	};
+
+	/** Rebuild indices from the provider if GetUpdateCounter() has changed.
+	 * Converts event times from ProfileTime to RecordingTime using the IFrameProvider. */
+	void Update(const IMassTraceProvider& Provider,
+		const TraceServices::IAnalysisSession& Session);
+
+	/** Get all events for a specific entity. Returns empty array if entity not found. */
+	TConstArrayView<FCachedEntityEventRecord> GetEntityEvents(uint64 EntityId) const;
+
+	/** Get all entity IDs that have events referencing a specific archetype. */
+	const TSet<uint64>* GetEntitiesForArchetype(uint64 ArchetypeId) const;
+
+	/** Get all archetype IDs that appear in the event data. */
+	const TSet<uint64>& GetActiveArchetypeIds() const
+	{
+		FReadScopeLock ReadLock(CacheLock);
+		return ActiveArchetypeIds;
+	}
+
+	/** Whether any data has been cached. */
+	bool HasData() const
+	{
+		FReadScopeLock ReadLock(CacheLock);
+		return EntityEvents.Num() > 0;
+	}
+
+private:
+	/** Convert event times from ProfileTime to ElapsedTime (recording time). */
+	void ConvertTimesToRecordingTime(const TraceServices::IAnalysisSession& Session);
+
+	/** Entity ID -> array of events, sorted by ascending time after Update() */
+	TMap<uint64, TArray<FCachedEntityEventRecord>> EntityEvents;
+
+	/** Archetype ID -> set of entity IDs that reference it */
+	TMap<uint64, TSet<uint64>> ArchetypeEntities;
+
+	/** All archetype IDs that appear in the event data */
+	TSet<uint64> ActiveArchetypeIds;
+
+	/** Last seen update counter from the provider */
+	uint64 LastUpdateCounter = 0;
+
+	/** Protects all cache state from concurrent read/write access */
+	mutable FRWLock CacheLock;
+};
+} // namespace UE::Mass::Trace
diff --git a/Engine/Source/Editor/MassEntityEditor/Private/Trace/MassEntityTrack.cpp b/Engine/Source/Editor/MassEntityEditor/Private/Trace/MassEntityTrack.cpp
new file mode 100644
--- /dev/null
+++ b/Engine/Source/Editor/MassEntityEditor/Private/Trace/MassEntityTrack.cpp
@@ -0,0 +1,301 @@
+// Copyright Epic Games, Inc. All Rights Reserved.
+
+#include "MassEntityTrack.h"
+#include "Common/ProviderLock.h"
+#include "IRewindDebugger.h"
+#include "Mass/EntityHandle.h"
+#include "MassEntityEventCache.h"
+#include "SMassEntityDetailsView.h"
+#include "Trace/MassTraceTypes.h"
+#include "TraceServices/Model/AnalysisSession.h"
+
+#define LOCTEXT_NAMESPACE "MassEntityTrack"
+
+namespace UE::Mass::Trace
+{
+static constexpr FLinearColor CreatedColor(0.2f, 0.9f, 0.2f, 1.0f);
+static constexpr FLinearColor ArchetypeChangeColor(0.9f, 0.9f, 0.2f, 1.0f);
+static constexpr FLinearColor DestroyedColor(0.9f, 0.2f, 0.2f, 1.0f);
+
+static FLinearColor ColorFromArchetypeId(const uint64 ArchetypeId)
+{
+	// Thorough bit mixing to produce distinct hues from pointer-based IDs
+	// that tend to be numerically close to each other.
+	uint64 H = ArchetypeId;
+	H ^= H >> 33;
+	H *= 0xff51afd7ed558ccdULL;
+	H ^= H >> 33;
+	H *= 0xc4ceb9fe1a85ec53ULL;
+	H ^= H >> 33;
+
+	// Use golden ratio to spread hues evenly
+	constexpr float GoldenRatio = 0.618033988749895f;
+	const float Hue = FMath::Frac(static_cast<float>(H & 0xFFFF) / 65535.0f + GoldenRatio);
+
+	// Vary saturation and value slightly for additional differentiation
+	const uint8 Saturation = 180 + static_cast<uint8>((H >> 16) % 40); // 180-219
+	const uint8 Value = 190 + static_cast<uint8>((H >> 24) % 50); // 190-239
+
+	return FLinearColor::MakeFromHSV8(
+		static_cast<uint8>(Hue * 255.0f),
+		Saturation,
+		Value
+	).CopyWithNewOpacity(0.55f);
+}
+
+static FText BuildArchetypeTooltip(
+	const IMassTraceProvider& Provider,
+	const uint64 ArchetypeID)
+{
+	const FArchetypeInfo* Info = nullptr;
+	{
+		TraceServices::FProviderReadScopeLock ReadLock(Provider);
+		Info = Provider.FindArchetypeById(ArchetypeID);
+	}
+	if (!Info || Info->Fragments.Num() == 0)
+	{
+		return FText();
+	}
+	FString Name;
+	constexpr int32 MaxDisplayedFragments = 3;
+	for (int32 Idx = 0; Idx < FMath::Min(MaxDisplayedFragments, Info->Fragments.Num()); ++Idx)
+	{
+		if (Idx > 0)
+		{
+			Name += TEXT(", ");
+		}
+		Name += Info->Fragments[Idx]->Name;
+	}
+	if (Info->Fragments.Num() > MaxDisplayedFragments)
+	{
+		Name += FString::Printf(TEXT(" (+%d)"), Info->Fragments.Num() - MaxDisplayedFragments);
+	}
+	return FText::FromString(MoveTemp(Name));
+}
+
+
+//----------------------------------------------------------------------//
+// FEntityTrackCreator
+//----------------------------------------------------------------------//
+
+FName FEntityTrackCreator::GetTargetTypeNameInternal() const
+{
+	static const FName Name(TEXT("MassEntityHandle"));
+	return Name;
+}
+
+FName FEntityTrackCreator::GetNameInternal() const
+{
+	static const FName Name(TEXT("MassEntity"));
+	return Name;
+}
+
+void FEntityTrackCreator::GetTrackTypesInternal(TArray<RewindDebugger::FRewindDebuggerTrackType>& Types) const
+{
+	Types.Add({ FName(TEXT("MassEntity")), LOCTEXT("MassEntity", "Mass Entity") });
+}
+
+TSharedPtr<RewindDebugger::FRewindDebuggerTrack> FEntityTrackCreator::CreateTrackInternal(const RewindDebugger::FObjectId& InObjectId) const
+{
+	const TraceServices::IAnalysisSession* Session = IRewindDebugger::Instance()->GetAnalysisSession();
+	if (Session)
+	{
+		const IMassTraceProvider& Provider = ReadMassTraceProvider(*Session);
+		Cache->Update(Provider, *Session);
+	}
+
+	return MakeShared<FEntityTrack>(InObjectId.GetMainId(), Cache);
+}
+
+bool FEntityTrackCreator::HasDebugInfoInternal(const RewindDebugger::FObjectId& InObjectId) const
+{
+	const TraceServices::IAnalysisSession* Session = IRewindDebugger::Instance()->GetAnalysisSession();
+	if (!Session)
+	{
+		return false;
+	}
+
+	const IMassTraceProvider& Provider = ReadMassTraceProvider(*Session);
+	Cache->Update(Provider, *Session);
+	return Cache->GetEntityEvents(InObjectId.GetMainId()).Num() > 0;
+}
+
+//----------------------------------------------------------------------//
+// FEntityTrack
+//----------------------------------------------------------------------//
+FEntityTrack::FEntityTrack(const uint64 InEntityId, TSharedRef<FEntityEventCache> InCache)
+	: EntityId(InEntityId)
+	, Cache(MoveTemp(InCache))
+{
+	EventData = MakeShared<SEventTimelineView::FTimelineEventData>();
+	Icon = FSlateIcon("EditorStyle", "ClassIcon.Object");
+
+	const FMassEntityHandle EntityHandle(FMassEntityHandle::FromNumber(InEntityId));
+	DisplayName = FText::Format(LOCTEXT("EntityName", "Entity {0}"), FText::FromString(LexToString(EntityHandle)));
+}
+
+bool FEntityTrack::UpdateInternal()
+{
+	// Always update the details view scrub time, even when throttling timeline rebuilds
+	if (DetailsView.IsValid())
+	{
+		DetailsView->SetScrubTime(IRewindDebugger::Instance()->GetScrubTime());
+	}
+
+	TicksSinceEventRebuild++;
+	if (TicksSinceEventRebuild <= 10)
+	{
+		return false;
+	}
+	TicksSinceEventRebuild = 0;
+
+	EventData->EventPoints.Reset();
+	EventData->EventWindows.Reset();
+
+	const TConstArrayView<FEntityEventCache::FCachedEntityEventRecord> Events = Cache->GetEntityEvents(EntityId);
+	if (Events.Num() == 0)
+	{
+		return false;
+	}
+
+	const TraceServices::IAnalysisSession* Session = IRewindDebugger::Instance()->GetAnalysisSession();
+	if (Session == nullptr)
+	{
+		return false;
+	}
+
+	const TRange<double> ViewRange = IRewindDebugger::Instance()->GetCurrentViewRange();
+	const double StartTime = ViewRange.GetLowerBoundValue();
+	const double EndTime = ViewRange.GetUpperBoundValue();
+
+	const IMassTraceProvider& Provider = ReadMassTraceProvider(*Session);
+
+	double WindowStart = -1.0;
+	uint64 CurrentArchetypeID = 0;
+
+	// Closes the current archetype window by adding it to EventData->EventWindows.
+	// Uses the provider to build a tooltip from the archetype's fragment names.
+	// bStillActive marks the window as the entity's current (not yet closed) archetype.
+	auto CloseCurrentWindow = [this, &Provider, &WindowStart, &CurrentArchetypeID](const double WindowEnd, const bool bStillActive)
+	{
+		if (WindowStart < 0.0)
+		{
+			return;
+		}
+		const FText Tooltip(BuildArchetypeTooltip(Provider, CurrentArchetypeID));
+		FLinearColor WindowColor = ColorFromArchetypeId(CurrentArchetypeID);
+		if (bStillActive)
+		{
+			// Brighter and more opaque for the still-active window
+			WindowColor = WindowColor.CopyWithNewOpacity(0.85f);
+		}
+		EventData->EventWindows.Add({
+			WindowStart, WindowEnd,
+			Tooltip,
+			bStillActive ? LOCTEXT("Active", "(active)") : FText(),
+			WindowColor
+		});
+		WindowStart = -1.0;
+		CurrentArchetypeID = 0;
+	};
+
+	for (const FEntityEventCache::FCachedEntityEventRecord& Event : Events)
+	{
+		const bool bInViewRange = Event.Time >= StartTime && Event.Time <= EndTime;
+
+		switch (Event.Operation)
+		{
+		case EEntityEventType::Created:
+			WindowStart = Event.Time;
+			CurrentArchetypeID = Event.ArchetypeID;
+			if (bInViewRange)
+			{
+				const FText Tooltip(BuildArchetypeTooltip(Provider, Event.ArchetypeID));
+				EventData->EventPoints.Add({Event.Time, LOCTEXT("Created", "Created"), Tooltip, CreatedColor});
+			}
+			break;
+
+		case EEntityEventType::ArchetypeChange:
+			// Close previous archetype window, start new one
+			CloseCurrentWindow(Event.Time, /*bStillActive*/ false);
+			WindowStart = Event.Time;
+			CurrentArchetypeID = Event.ArchetypeID;
+			if (bInViewRange)
+			{
+				const FText Tooltip(BuildArchetypeTooltip(Provider, Event.ArchetypeID));
+				EventData->EventPoints.Add({Event.Time, LOCTEXT("ArchetypeChange", "Archetype Change"), Tooltip, ArchetypeChangeColor});
+			}
+			break;
+
+		case EEntityEventType::Destroyed:
+			CloseCurrentWindow(Event.Time, /*bStillActive*/ false);
+			if (bInViewRange)
+			{
+				EventData->EventPoints.Add({Event.Time, LOCTEXT("Destroyed", "Destroyed"), FText(), DestroyedColor});
+			}
+			break;
+
+		default:
+			break;
+		}
+	}
+
+	// Entity still alive — snap window to recording duration and mark as active
+	if (WindowStart >= 0.0)
+	{
+		const double RecordingEnd = IRewindDebugger::Instance()->GetRecordingDuration();
+		CloseCurrentWindow(RecordingEnd, /*bStillActive*/ true);
+	}
+
+	return false;
+}
+
+TSharedPtr<SWidget> FEntityTrack::GetTimelineViewInternal()
+{
+	return SNew(SEventTimelineView)
+		.ViewRange_Lambda([]()
+		{
+			return IRewindDebugger::Instance()->GetCurrentViewRange();
+		})
+		.EventData_Raw(this, &FEntityTrack::GetEventData);
+}
+
+TSharedPtr<SWidget> FEntityTrack::GetDetailsViewInternal()
+{
+	if (!DetailsView.IsValid())
+	{
+		SAssignNew(DetailsView, SMassEntityDetailsView)
+			.EntityId(EntityId)
+			.Cache(&Cache.Get());
+	}
+	return DetailsView;
+}
+
+TSharedPtr<SEventTimelineView::FTimelineEventData> FEntityTrack::GetEventData() const
+{
+	return EventData;
+}
+
+FSlateIcon FEntityTrack::GetIconInternal()
+{
+	return Icon;
+}
+
+FName FEntityTrack::GetNameInternal() const
+{
+	return FName(*FString::Printf(TEXT("MassEntity_%llu"), EntityId));
+}
+
+FText FEntityTrack::GetDisplayNameInternal() const
+{
+	return DisplayName;
+}
+
+bool FEntityTrack::HasDebugDataInternal() const
+{
+	return Cache->GetEntityEvents(EntityId).Num() > 0;
+}
+
+} // namespace UE::Mass::Trace
+
+#undef LOCTEXT_NAMESPACE
diff --git a/Engine/Source/Editor/MassEntityEditor/Private/Trace/MassEntityTrack.h b/Engine/Source/Editor/MassEntityEditor/Private/Trace/MassEntityTrack.h
new file mode 100644
--- /dev/null
+++ b/Engine/Source/Editor/MassEntityEditor/Private/Trace/MassEntityTrack.h
@@ -0,0 +1,73 @@
+// Copyright Epic Games, Inc. All Rights Reserved.
+
+#pragma once
+
+#include "IRewindDebuggerTrackCreator.h"
+#include "MassEntityEventCache.h"
+#include "RewindDebuggerTrack.h"
+#include "SEventTimelineView.h"
+
+class SMassEntityDetailsView;
+
+namespace UE::Mass::Trace
+{
+
+/**
+ * Per-entity leaf track showing lifecycle events (created, archetype change, destroyed)
+ * on a timeline with color-coded points and lifetime windows.
+ */
+class FEntityTrack : public RewindDebugger::FRewindDebuggerTrack
+{
+public:
+	FEntityTrack(uint64 InEntityId, TSharedRef<FEntityEventCache> InCache);
+
+	uint64 GetEntityId() const
+	{
+		return EntityId;
+	}
+
+private:
+	//~ FRewindDebuggerTrack interface
+	virtual bool UpdateInternal() override;
+	virtual TSharedPtr<SWidget> GetTimelineViewInternal() override;
+	virtual TSharedPtr<SWidget> GetDetailsViewInternal() override;
+	virtual FSlateIcon GetIconInternal() override;
+	virtual FName GetNameInternal() const override;
+	virtual FText GetDisplayNameInternal() const override;
+	virtual uint64 GetObjectIdInternal() const override
+	{
+		return EntityId;
+	}
+	virtual bool HasDebugDataInternal() const override;
+
+	TSharedPtr<SEventTimelineView::FTimelineEventData> GetEventData() const;
+
+	uint64 EntityId = 0;
+	TSharedRef<FEntityEventCache> Cache;
+
+	mutable TSharedPtr<SEventTimelineView::FTimelineEventData> EventData;
+	mutable int32 TicksSinceEventRebuild = 0;
+
+	TSharedPtr<SMassEntityDetailsView> DetailsView;
+
+	FText DisplayName;
+	FSlateIcon Icon;
+};
+
+/**
+ * Track creator for individual Mass entity tracks.
+ * Creates FEntityTrack under FMassEntityHandle object tracks.
+ */
+class FEntityTrackCreator : public RewindDebugger::IRewindDebuggerTrackCreator
+{
+private:
+	virtual FName GetTargetTypeNameInternal() const override;
+	virtual FName GetNameInternal() const override;
+	virtual void GetTrackTypesInternal(TArray<RewindDebugger::FRewindDebuggerTrackType>& Types) const override;
+	virtual TSharedPtr<RewindDebugger::FRewindDebuggerTrack> CreateTrackInternal(const RewindDebugger::FObjectId& InObjectId) const override;
+	virtual bool HasDebugInfoInternal(const RewindDebugger::FObjectId& InObjectId) const override;
+
+	mutable TSharedRef<FEntityEventCache> Cache = MakeShared<FEntityEventCache>();
+};
+
+} // UE::Mass::Trace
diff --git a/Engine/Source/Editor/MassEntityEditor/Private/Trace/MassRewindDebuggerIntegration.cpp b/Engine/Source/Editor/MassEntityEditor/Private/Trace/MassRewindDebuggerIntegration.cpp
new file mode 100644
--- /dev/null
+++ b/Engine/Source/Editor/MassEntityEditor/Private/Trace/MassRewindDebuggerIntegration.cpp
@@ -0,0 +1,14 @@
+// Copyright Epic Games, Inc. All Rights Reserved.
+
+#include "MassRewindDebuggerIntegration.h"
+#include "Trace/Trace.h"
+
+void FMassRewindDebuggerRuntimeExtension::RecordingStarted()
+{
+	UE::Trace::ToggleChannel(TEXT("MassChannel"), true);
+}
+
+void FMassRewindDebuggerRuntimeExtension::RecordingStopped()
+{
+	UE::Trace::ToggleChannel(TEXT("MassChannel"), false);
+}
diff --git a/Engine/Source/Editor/MassEntityEditor/Private/Trace/MassRewindDebuggerIntegration.h b/Engine/Source/Editor/MassEntityEditor/Private/Trace/MassRewindDebuggerIntegration.h
new file mode 100644
--- /dev/null
+++ b/Engine/Source/Editor/MassEntityEditor/Private/Trace/MassRewindDebuggerIntegration.h
@@ -0,0 +1,16 @@
+// Copyright Epic Games, Inc. All Rights Reserved.
+
+#pragma once
+
+#include "RewindDebuggerRuntimeInterface/IRewindDebuggerRuntimeExtension.h"
+
+/**
+ * Runtime extension that enables/disables the Mass trace channel
+ * when RewindDebugger recording starts/stops.
+ */
+class FMassRewindDebuggerRuntimeExtension final : public IRewindDebuggerRuntimeExtension
+{
+protected:
+	virtual void RecordingStarted() override;
+	virtual void RecordingStopped() override;
+};
diff --git a/Engine/Source/Editor/MassEntityEditor/Private/Trace/SMassEntityDetailsView.cpp b/Engine/Source/Editor/MassEntityEditor/Private/Trace/SMassEntityDetailsView.cpp
new file mode 100644
--- /dev/null
+++ b/Engine/Source/Editor/MassEntityEditor/Private/Trace/SMassEntityDetailsView.cpp
@@ -0,0 +1,526 @@
+// Copyright Epic Games, Inc. All Rights Reserved.
+
+#include "SMassEntityDetailsView.h"
+#include "Common/ProviderLock.h"
+#include "IRewindDebugger.h"
+#include "MassEntityEventCache.h"
+#include "Trace/MassTraceTypes.h"
+#include "TraceServices/Model/AnalysisSession.h"
+#include "Widgets/Text/STextBlock.h"
+#include "Widgets/Views/SHeaderRow.h"
+#include "Widgets/Views/SListView.h"
+
+#define LOCTEXT_NAMESPACE "MassEntityDetailsView"
+
+namespace UE::MassEntityDetailsView::Private
+{
+static constexpr FLinearColor AddedColor(0.2f, 0.9f, 0.2f, 1.0f);
+static constexpr FLinearColor RemovedColor(0.9f, 0.3f, 0.3f, 1.0f);
+static constexpr FLinearColor UnchangedColor(0.8f, 0.8f, 0.8f, 1.0f);
+static constexpr FLinearColor ReaddedColor(0.6f, 0.6f, 0.9f, 1.0f);
+static constexpr FLinearColor InfoColor(0.5f, 0.5f, 0.5f, 1.0f);
+
+static const FName PreviousColumnName(TEXT("Previous"));
+static const FName CurrentColumnName(TEXT("Current"));
+static const FName TypeColumnName(TEXT("Type"));
+static const FName SizeColumnName(TEXT("Size"));
+
+} // namespace UE::MassEntityDetailsView::Private
+
+/**
+ * Row widget for the diff table.
+ */
+class SMassArchetypeDiffRow : public SMultiColumnTableRow<TSharedPtr<SMassEntityDetailsView::FFragmentEntry>>
+{
+	using FEntry = SMassEntityDetailsView::FFragmentEntry;
+	using EDiffStatus = SMassEntityDetailsView::EDiffStatus;
+
+public:
+	SLATE_BEGIN_ARGS(SMassArchetypeDiffRow) {}
+		SLATE_ARGUMENT(TSharedPtr<FEntry>, Entry)
+	SLATE_END_ARGS()
+
+	void Construct(const FArguments& InArgs, const TSharedRef<STableViewBase>& InOwnerTableView)
+	{
+		Entry = InArgs._Entry;
+		SMultiColumnTableRow::Construct(SMultiColumnTableRow::FArguments(), InOwnerTableView);
+	}
+
+	virtual TSharedRef<SWidget> GenerateWidgetForColumn(const FName& ColumnName) override
+	{
+		using namespace UE::MassEntityDetailsView::Private;
+
+		if (!Entry.IsValid() || !Entry->Info)
+		{
+			return SNullWidget::NullWidget;
+		}
+
+		const FText FragName = FText::FromString(Entry->Info->Name);
+
+		if (ColumnName == PreviousColumnName)
+		{
+			// Removed elements show in the Previous column with - prefix and red
+			if (Entry->Status == EDiffStatus::Removed)
+			{
+				return SNew(STextBlock)
+					.Text(FormatWithStep(LOCTEXT("PrefixRemove", "-"), FragName, Entry->StepIndex))
+					.ColorAndOpacity(RemovedColor);
+			}
+			// Show fragment name if it was in the previous archetype
+			if (Entry->bHasPrevious)
+			{
+				return SNew(STextBlock)
+					.Text(FragName)
+					.ColorAndOpacity(UnchangedColor);
+			}
+			return SNullWidget::NullWidget;
+		}
+
+		if (ColumnName == CurrentColumnName)
+		{
+			// Removed elements are shown in Previous column, not here
+			if (Entry->Status == EDiffStatus::Removed)
+			{
+				return SNullWidget::NullWidget;
+			}
+
+			FText DisplayText;
+			FLinearColor Color = UnchangedColor;
+
+			switch (Entry->Status)
+			{
+			case EDiffStatus::Unchanged:
+				DisplayText = FragName;
+				break;
+			case EDiffStatus::Added:
+				Color = AddedColor;
+				DisplayText = FormatWithStep(LOCTEXT("PrefixAdd", "+"), FragName, Entry->StepIndex);
+				break;
+			case EDiffStatus::Readded:
+				Color = ReaddedColor;
+				DisplayText = FormatWithStep(LOCTEXT("PrefixReadd", "~"), FragName, Entry->StepIndex);
+				break;
+			default:
+				break;
+			}
+
+			return SNew(STextBlock)
+				.Text(DisplayText)
+				.ColorAndOpacity(Color);
+		}
+
+		if (ColumnName == TypeColumnName)
+		{
+			return SNew(STextBlock)
+				.Text(GetTypeText())
+				.ColorAndOpacity(InfoColor);
+		}
+
+		if (ColumnName == SizeColumnName)
+		{
+			return SNew(STextBlock)
+				.Text(FText::Format(LOCTEXT("SizeFmt", "{0}b"), FText::AsNumber(Entry->Info->Size)))
+				.ColorAndOpacity(InfoColor);
+		}
+
+		return SNullWidget::NullWidget;
+	}
+
+private:
+	static FText FormatWithStep(const FText& Prefix, const FText& Name, const int32 StepIndex)
+	{
+		if (StepIndex > 0)
+		{
+			return FText::Format(LOCTEXT("StepFmt", "{0} {1} [step {2}]"), Prefix, Name, FText::AsNumber(StepIndex));
+		}
+		return FText::Format(LOCTEXT("NoStepFmt", "{0} {1}"), Prefix, Name);
+	}
+
+	FText GetTypeText() const
+	{
+		using namespace UE::Mass::Trace;
+
+		switch (Entry->Info->Type)
+		{
+		case EFragmentType::Fragment:
+			return LOCTEXT("TFragment", "Fragment");
+		case EFragmentType::Tag:
+			return LOCTEXT("TTag", "Tag");
+		case EFragmentType::Shared:
+			return LOCTEXT("TShared", "Shared");
+		case EFragmentType::Sparse:
+			return LOCTEXT("TSparse", "Sparse");
+		case EFragmentType::ConstShared:
+			return LOCTEXT("TConstShared", "Const Shared");
+		case EFragmentType::SparseTag:
+			return LOCTEXT("TSparseTag", "Sparse Tag");
+		default:
+			return LOCTEXT("TUnknown", "Unknown");
+		}
+	}
+
+	TSharedPtr<FEntry> Entry;
+};
+
+//----------------------------------------------------------------------//
+// SMassEntityDetailsView
+//----------------------------------------------------------------------//
+
+void SMassEntityDetailsView::Construct(const FArguments& InArgs)
+{
+	using namespace UE::MassEntityDetailsView::Private;
+	EntityId = InArgs._EntityId;
+	Cache = InArgs._Cache;
+
+	HeaderRow = SNew(SHeaderRow)
+		+ SHeaderRow::Column(PreviousColumnName)
+		.DefaultLabel(LOCTEXT("ColPrev", "Previous"))
+		.FillWidth(1.0f)
+
+		+ SHeaderRow::Column(CurrentColumnName)
+		.DefaultLabel(LOCTEXT("ColCurr", "Current"))
+		.FillWidth(1.0f)
+
+		+ SHeaderRow::Column(TypeColumnName)
+		.DefaultLabel(LOCTEXT("ColType", "Type"))
+		.FillWidth(0.5f)
+
+		+ SHeaderRow::Column(SizeColumnName)
+		.DefaultLabel(LOCTEXT("ColSize", "Size"))
+		.FillWidth(0.3f);
+
+	ChildSlot
+	[
+		SNew(SVerticalBox)
+		+ SVerticalBox::Slot()
+		.AutoHeight()
+		.Padding(4.0f, 2.0f)
+		[
+			SAssignNew(MergeInfoText, STextBlock)
+			.Font(FCoreStyle::GetDefaultFontStyle("Italic", 8))
+			.ColorAndOpacity(FSlateColor(FLinearColor(0.5f, 0.5f, 0.5f)))
+			.Visibility(EVisibility::Collapsed)
+		]
+		+ SVerticalBox::Slot()
+		.FillHeight(1.0f)
+		[
+			SAssignNew(ListView, SListView<TSharedPtr<FFragmentEntry>>)
+			.ListItemsSource(&FragmentEntries)
+			.OnGenerateRow(this, &SMassEntityDetailsView::OnGenerateRow)
+			.SelectionMode(ESelectionMode::None)
+			.HeaderRow(HeaderRow)
+		]
+	];
+}
+
+void SMassEntityDetailsView::SetScrubTime(const double Time)
+{
+	CurrentScrubTime = Time;
+}
+
+void SMassEntityDetailsView::Tick(const FGeometry& AllottedGeometry, const double InCurrentTime, const float InDeltaTime)
+{
+	SCompoundWidget::Tick(AllottedGeometry, InCurrentTime, InDeltaTime);
+
+	const double ScrubTime = IRewindDebugger::Instance()->GetScrubTime();
+	SetScrubTime(ScrubTime);
+
+	uint64 CurrentArchetype = 0;
+	uint64 PreviousArchetype = 0;
+	int32 MergeCount = 0;
+	TArray<uint64> BatchArchetypeIds;
+	ResolveArchetypesAtScrubTime(CurrentArchetype, PreviousArchetype, MergeCount, BatchArchetypeIds);
+
+	if (CurrentArchetype != LastDisplayedArchetypeId ||
+		PreviousArchetype != LastDisplayedPreviousArchetypeId ||
+		MergeCount != LastDisplayedMergeCount)
+	{
+		LastDisplayedArchetypeId = CurrentArchetype;
+		LastDisplayedPreviousArchetypeId = PreviousArchetype;
+		LastDisplayedMergeCount = MergeCount;
+		LastBatchArchetypeIds = MoveTemp(BatchArchetypeIds);
+		RebuildFragmentList();
+	}
+}
+
+TSharedRef<ITableRow> SMassEntityDetailsView::OnGenerateRow(
+	TSharedPtr<FFragmentEntry> Entry,
+	const TSharedRef<STableViewBase>& OwnerTable)
+{
+	return SNew(SMassArchetypeDiffRow, OwnerTable).Entry(Entry);
+}
+
+void SMassEntityDetailsView::ResolveArchetypesAtScrubTime(uint64& OutCurrentArchetype, uint64& OutPreviousArchetype,
+	int32& OutMergeCount, TArray<uint64>& OutBatchArchetypeIds) const
+{
+	using namespace UE::Mass::Trace;
+
+	OutCurrentArchetype = 0;
+	OutPreviousArchetype = 0;
+	OutMergeCount = 0;
+	OutBatchArchetypeIds.Reset();
+
+	if (!Cache)
+	{
+		return;
+	}
+
+	TConstArrayView<FEntityEventCache::FCachedEntityEventRecord> Events = Cache->GetEntityEvents(EntityId);
+	uint32 LastFrameIndex = MAX_uint32;
+	uint64 ArchetypeBeforeCurrentBatch = 0;
+	int32 CurrentBatchChangeCount = 0;
+	TArray<uint64> CurrentBatchArchetypeIds;
+
+	for (const FEntityEventCache::FCachedEntityEventRecord& Record : Events)
+	{
+		if (Record.Time > CurrentScrubTime + UE_DOUBLE_SMALL_NUMBER)
+		{
+			break;
+		}
+
+		if (Record.FrameIndex != LastFrameIndex)
+		{
+			ArchetypeBeforeCurrentBatch = OutCurrentArchetype;
+			LastFrameIndex = Record.FrameIndex;
+			CurrentBatchChangeCount = 0;
+			CurrentBatchArchetypeIds.Reset();
+		}
+
+		switch (Record.Operation)
+		{
+		case EEntityEventType::Created:
+			ArchetypeBeforeCurrentBatch = 0;
+			OutCurrentArchetype = Record.ArchetypeID;
+			CurrentBatchArchetypeIds.Add(Record.ArchetypeID);
+			CurrentBatchChangeCount++;
+			break;
+		case EEntityEventType::ArchetypeChange:
+			OutCurrentArchetype = Record.ArchetypeID;
+			CurrentBatchArchetypeIds.Add(Record.ArchetypeID);
+			CurrentBatchChangeCount++;
+			break;
+		case EEntityEventType::Destroyed:
+			OutCurrentArchetype = 0;
+			CurrentBatchChangeCount++;
+			break;
+		default:
+			break;
+		}
+	}
+
+	OutPreviousArchetype = ArchetypeBeforeCurrentBatch;
+	OutMergeCount = CurrentBatchChangeCount;
+	OutBatchArchetypeIds = MoveTemp(CurrentBatchArchetypeIds);
+}
+
+void SMassEntityDetailsView::RebuildFragmentList()
+{
+	using namespace UE::Mass::Trace;
+
+	FragmentEntries.Reset();
+
+	const uint64 ArchetypeAtScrubTime = LastDisplayedArchetypeId;
+	const uint64 PreviousArchetypeId = LastDisplayedPreviousArchetypeId;
+
+	// Update merge info text
+	if (LastDisplayedMergeCount > 1 && PreviousArchetypeId != 0)
+	{
+		MergeInfoText->SetText(FText::Format(
+			LOCTEXT("MergedNote", "Merged {0} archetype changes in one frame"),
+			FText::AsNumber(LastDisplayedMergeCount)));
+		MergeInfoText->SetVisibility(EVisibility::Visible);
+	}
+	else
+	{
+		MergeInfoText->SetVisibility(EVisibility::Collapsed);
+	}
+
+	// Show/hide Previous column based on whether there is a previous archetype
+	HeaderRow->SetShowGeneratedColumn(UE::MassEntityDetailsView::Private::PreviousColumnName, PreviousArchetypeId != 0);
+
+	if (ArchetypeAtScrubTime == 0)
+	{
+		if (ListView.IsValid())
+		{
+			ListView->RequestListRefresh();
+		}
+		return;
+	}
+
+	const TraceServices::IAnalysisSession* Session = IRewindDebugger::Instance()->GetAnalysisSession();
+	if (!Session)
+	{
+		return;
+	}
+
+	const IMassTraceProvider& MassTraceProvider = ReadMassTraceProvider(*Session);
+
+	const FArchetypeInfo* CurrentArchetype = nullptr;
+	const FArchetypeInfo* PreviousArchetype = nullptr;
+	{
+		TraceServices::FProviderReadScopeLock ReadLock(MassTraceProvider);
+		CurrentArchetype = MassTraceProvider.FindArchetypeById(ArchetypeAtScrubTime);
+		if (PreviousArchetypeId != 0)
+		{
+			PreviousArchetype = MassTraceProvider.FindArchetypeById(PreviousArchetypeId);
+		}
+	}
+
+	if (!CurrentArchetype)
+	{
+		if (ListView.IsValid())
+		{
+			ListView->RequestListRefresh();
+		}
+		return;
+	}
+
+	// Build fragment ID sets
+	TSet<uint64> PreviousFragmentIds;
+	TSet<uint64> CurrentFragmentIds;
+
+	if (PreviousArchetype)
+	{
+		for (const FFragmentInfo* Fragment : PreviousArchetype->Fragments)
+		{
+			if (Fragment)
+			{
+				PreviousFragmentIds.Add(Fragment->Id);
+			}
+		}
+	}
+
+	for (const FFragmentInfo* Fragment : CurrentArchetype->Fragments)
+	{
+		if (Fragment)
+		{
+			CurrentFragmentIds.Add(Fragment->Id);
+		}
+	}
+
+	// Step attribution for merged batches
+	TMap<uint64, int32> AddedAtStep;
+	TMap<uint64, int32> RemovedAtStep;
+	TSet<uint64> ReaddedIds;
+
+	if (LastDisplayedMergeCount > 1 && PreviousArchetype)
+	{
+		TraceServices::FProviderReadScopeLock ReadLock(MassTraceProvider);
+		TSet<uint64> StepFragmentIds = PreviousFragmentIds;
+
+		for (int32 StepIdx = 0; StepIdx < LastBatchArchetypeIds.Num(); ++StepIdx)
+		{
+			const FArchetypeInfo* StepArchetype = MassTraceProvider.FindArchetypeById(LastBatchArchetypeIds[StepIdx]);
+			if (!StepArchetype)
+			{
+				continue;
+			}
+
+			TSet<uint64> NextFragmentIds;
+			for (const FFragmentInfo* Fragment : StepArchetype->Fragments)
+			{
+				if (Fragment)
+				{
+					NextFragmentIds.Add(Fragment->Id);
+				}
+			}
+
+			for (uint64 Id : NextFragmentIds)
+			{
+				if (!StepFragmentIds.Contains(Id))
+				{
+					if (RemovedAtStep.Contains(Id))
+					{
+						ReaddedIds.Add(Id);
+					}
+					if (!AddedAtStep.Contains(Id))
+					{
+						AddedAtStep.Add(Id, StepIdx + 1);
+					}
+				}
+			}
+
+			for (uint64 Id : StepFragmentIds)
+			{
+				if (!NextFragmentIds.Contains(Id))
+				{
+					if (AddedAtStep.Contains(Id))
+					{
+						ReaddedIds.Add(Id);
+					}
+					if (!RemovedAtStep.Contains(Id))
+					{
+						RemovedAtStep.Add(Id, StepIdx + 1);
+					}
+				}
+			}
+
+			StepFragmentIds = MoveTemp(NextFragmentIds);
+		}
+	}
+
+	// Build entries
+	for (const FFragmentInfo* Fragment : CurrentArchetype->Fragments)
+	{
+		if (!Fragment)
+		{
+			continue;
+		}
+
+		const bool bInPrevious = PreviousArchetype && PreviousFragmentIds.Contains(Fragment->Id);
+		const bool bIsReadded = ReaddedIds.Contains(Fragment->Id) && bInPrevious;
+
+		TSharedPtr<FFragmentEntry> Entry = MakeShared<FFragmentEntry>();
+		Entry->Info = Fragment;
+		Entry->bHasPrevious = bInPrevious;
+
+		if (bIsReadded)
+		{
+			Entry->Status = EDiffStatus::Readded;
+			Entry->StepIndex = RemovedAtStep.FindRef(Fragment->Id);
+		}
+		else if (!bInPrevious && PreviousArchetype)
+		{
+			Entry->Status = EDiffStatus::Added;
+			Entry->StepIndex = AddedAtStep.FindRef(Fragment->Id);
+		}
+		else
+		{
+			Entry->Status = EDiffStatus::Unchanged;
+		}
+
+		FragmentEntries.Add(MoveTemp(Entry));
+	}
+
+	// Add removed fragments
+	if (PreviousArchetype)
+	{
+		for (const FFragmentInfo* Fragment : PreviousArchetype->Fragments)
+		{
+			if (!Fragment || CurrentFragmentIds.Contains(Fragment->Id))
+			{
+				continue;
+			}
+
+			TSharedPtr<FFragmentEntry> Entry = MakeShared<FFragmentEntry>();
+			Entry->Info = Fragment;
+			Entry->Status = EDiffStatus::Removed;
+			Entry->StepIndex = RemovedAtStep.FindRef(Fragment->Id);
+			Entry->bHasPrevious = true;
+			FragmentEntries.Add(MoveTemp(Entry));
+		}
+	}
+
+	// Sort alphabetically
+	FragmentEntries.Sort([](const TSharedPtr<FFragmentEntry>& A, const TSharedPtr<FFragmentEntry>& B)
+	{
+		return A->Info->Name < B->Info->Name;
+	});
+
+	if (ListView.IsValid())
+	{
+		ListView->RequestListRefresh();
+	}
+}
+
+#undef LOCTEXT_NAMESPACE
diff --git a/Engine/Source/Editor/MassEntityEditor/Private/Trace/SMassEntityDetailsView.h b/Engine/Source/Editor/MassEntityEditor/Private/Trace/SMassEntityDetailsView.h
new file mode 100644
--- /dev/null
+++ b/Engine/Source/Editor/MassEntityEditor/Private/Trace/SMassEntityDetailsView.h
@@ -0,0 +1,72 @@
+// Copyright Epic Games, Inc. All Rights Reserved.
+
+#pragma once
+
+#include "Widgets/SCompoundWidget.h"
+#include "Widgets/DeclarativeSyntaxSupport.h"
+
+class STextBlock;
+class STableViewBase;
+class ITableRow;
+class SHeaderRow;
+template <typename ItemType> class SListView;
+
+namespace UE::Mass::Trace
+{
+struct FFragmentInfo;
+class FEntityEventCache;
+}
+
+/**
+ * Details widget for a single Mass entity.
+ * Shows a diff-style table of fragments between previous and current archetype.
+ */
+class SMassEntityDetailsView : public SCompoundWidget
+{
+public:
+	enum class EDiffStatus : uint8
+	{
+		Unchanged,
+		Added,
+		Removed,
+		Readded
+	};
+
+	struct FFragmentEntry
+	{
+		const UE::Mass::Trace::FFragmentInfo* Info = nullptr;
+		EDiffStatus Status = EDiffStatus::Unchanged;
+		int32 StepIndex = 0;
+		bool bHasPrevious = false;
+	};
+
+	SLATE_BEGIN_ARGS(SMassEntityDetailsView) {}
+		SLATE_ARGUMENT(uint64, EntityId)
+		SLATE_ARGUMENT(const UE::Mass::Trace::FEntityEventCache*, Cache)
+	SLATE_END_ARGS()
+
+	void Construct(const FArguments& InArgs);
+	void SetScrubTime(double Time);
+
+private:
+	virtual void Tick(const FGeometry& AllottedGeometry, const double InCurrentTime, const float InDeltaTime) override;
+
+	void RebuildFragmentList();
+	void ResolveArchetypesAtScrubTime(uint64& OutCurrentArchetype, uint64& OutPreviousArchetype,
+		int32& OutMergeCount, TArray<uint64>& OutBatchArchetypeIds) const;
+
+	TSharedRef<ITableRow> OnGenerateRow(TSharedPtr<FFragmentEntry> Entry, const TSharedRef<STableViewBase>& OwnerTable);
+
+	uint64 EntityId = 0;
+	const UE::Mass::Trace::FEntityEventCache* Cache = nullptr;
+	double CurrentScrubTime = 0.0;
+	uint64 LastDisplayedArchetypeId = 0;
+	uint64 LastDisplayedPreviousArchetypeId = 0;
+	int32 LastDisplayedMergeCount = 0;
+	TArray<uint64> LastBatchArchetypeIds;
+
+	TArray<TSharedPtr<FFragmentEntry>> FragmentEntries;
+	TSharedPtr<SListView<TSharedPtr<FFragmentEntry>>> ListView;
+	TSharedPtr<SHeaderRow> HeaderRow;
+	TSharedPtr<STextBlock> MergeInfoText;
+};
diff --git a/Engine/Source/Editor/MassEntityEditor/Public/MassEntityEditorModule.h b/Engine/Source/Editor/MassEntityEditor/Public/MassEntityEditorModule.h
--- a/Engine/Source/Editor/MassEntityEditor/Public/MassEntityEditorModule.h
+++ b/Engine/Source/Editor/MassEntityEditor/Public/MassEntityEditorModule.h
@@ -1,17 +1,24 @@
 // Copyright Epic Games, Inc. All Rights Reserved.
 
-#pragma  once
+#pragma once
 
+#include "Delegates/IDelegateInstance.h"
 #include "Modules/ModuleInterface.h"
 #include "Toolkits/AssetEditorToolkit.h"
-#include "Delegates/IDelegateInstance.h"
 
 #define UE_API MASSENTITYEDITOR_API
 
-class IMassEntityEditor;
-struct FGraphPanelNodeFactory;
 struct FGraphNodeClassHelper;
+struct FGraphPanelNodeFactory;
+class FMassEntitiesTrackCreator;
+class FMassEntityTrackCreator;
+class FMassRewindDebuggerRuntimeExtension;
+class IMassEntityEditor;
 class UWorld;
+namespace UE::Mass::Trace
+{
+class FEntityTrackCreator;
+}
 
 /**
 * The public interface to this module
@@ -47,6 +54,10 @@ protected:
 
 	TSharedPtr<FExtensibilityManager> MenuExtensibilityManager;
 	TSharedPtr<FExtensibilityManager> ToolBarExtensibilityManager;
+
+	//~ RewindDebugger integration
+	TUniquePtr<UE::Mass::Trace::FEntityTrackCreator> EntityTrackCreator;
+	TUniquePtr<FMassRewindDebuggerRuntimeExtension> RuntimeExtension;
 };
 
 #undef UE_API
diff --git a/Engine/Source/Editor/RewindDebuggerInterface/Private/SEventTimelineView.cpp b/Engine/Source/Editor/RewindDebuggerInterface/Private/SEventTimelineView.cpp
new file mode 100644
--- /dev/null
+++ b/Engine/Source/Editor/RewindDebuggerInterface/Private/SEventTimelineView.cpp
@@ -0,0 +1,131 @@
+// Copyright Epic Games, Inc. All Rights Reserved.
+
+#include "SEventTimelineView.h"
+#include "SSimpleTimeSlider.h"
+#include "Styling/AppStyle.h"
+
+#define LOCTEXT_NAMESPACE "SEventTimelineView"
+
+void SEventTimelineView::Construct(const SEventTimelineView::FArguments& InArgs)
+{
+	ViewRange = InArgs._ViewRange;
+	DesiredSize = InArgs._DesiredSize;
+	EventData = InArgs._EventData;
+
+	// set clipping on by default, since the OnPaint function is drawing outside the bounds
+	// Clipping = EWidgetClipping::ClipToBounds;
+	EventBrush = FAppStyle::GetBrush("Sequencer.KeyDiamond");
+	EventBorderBrush = FAppStyle::GetBrush("Sequencer.KeyDiamondBorder");
+}
+
+int32 SEventTimelineView::OnPaint(const FPaintArgs& Args, const FGeometry& AllottedGeometry, const FSlateRect& MyCullingRect, FSlateWindowElementList& OutDrawElements, const int32 LayerId, const FWidgetStyle& InWidgetStyle, const bool bParentEnabled) const
+{
+	const int32 NewLayer = PaintEvents(AllottedGeometry, MyCullingRect, OutDrawElements, LayerId, InWidgetStyle, bParentEnabled);
+
+	return FMath::Max(NewLayer, SCompoundWidget::OnPaint(Args, AllottedGeometry, MyCullingRect, OutDrawElements, NewLayer, InWidgetStyle, ShouldBeEnabled(bParentEnabled)));
+}
+
+FVector2D SEventTimelineView::ComputeDesiredSize(float) const
+{
+	return DesiredSize.Get();
+}
+
+int32 SEventTimelineView::PaintEvents(const FGeometry& AllottedGeometry, const FSlateRect& MyCullingRect, FSlateWindowElementList& OutDrawElements, int32 LayerId, const FWidgetStyle& InWidgetStyle, bool bParentEnabled) const
+{
+	const TRange<double> DebugTimeRange = ViewRange.Get();
+
+	const SSimpleTimeSlider::FScrubRangeToScreen RangeToScreen(DebugTimeRange, AllottedGeometry.GetLocalSize());
+
+	if (const TSharedPtr<FTimelineEventData> Event = EventData.Get())
+	{
+		const TConstArrayView<FTimelineEventData::FEventPoint> EventPoints = Event->EventPoints;
+		const int32 NumPoints = EventPoints.Num();
+
+		double PrevPointTime = 0;
+		int32 OverlappingPointCount = 0;
+
+		if (NumPoints > 0)
+		{
+			for (int32 i = 0; i < NumPoints; i++)
+			{
+				const FTimelineEventData::FEventPoint& Point = EventPoints[i];
+
+				FVector2D EventSize = EventBrush->GetImageSize();
+
+				float X = RangeToScreen.InputToLocalX(Point.Time);
+				X = X - EventSize.X / 2;
+				float Y = (AllottedGeometry.Size.Y - EventSize.Y) / 2;
+				if (Point.Time == PrevPointTime)
+				{
+					static constexpr int32 OverlapOffsetAmount = 2;
+					static constexpr int32 MaxOverlappingEvents = 2;
+
+					if (OverlappingPointCount > MaxOverlappingEvents)
+					{
+						continue;
+					}
+
+					OverlappingPointCount++;
+					Y += OverlapOffsetAmount * OverlappingPointCount;
+				}
+				else
+				{
+					OverlappingPointCount = 0;
+				}
+
+				PrevPointTime = Point.Time;
+
+				FSlateDrawElement::MakeBox(OutDrawElements, LayerId,
+					AllottedGeometry.ToPaintGeometry(EventSize, FSlateLayoutTransform(FVector2f(X, Y))), EventBrush, ESlateDrawEffect::None, Point.Color);
+
+				FSlateDrawElement::MakeBox(OutDrawElements, LayerId,
+					AllottedGeometry.ToPaintGeometry(EventSize, FSlateLayoutTransform(FVector2f(X, Y))), EventBorderBrush, ESlateDrawEffect::None, FLinearColor::Black);
+			}
+		}
+
+		const TConstArrayView<FTimelineEventData::FEventWindow> EventWindows = Event->EventWindows;
+		const int32 NumWindows = EventWindows.Num();
+
+		const FSlateBrush* Brush = FAppStyle::GetBrush("Sequencer.SectionArea.Background");
+
+		if (NumWindows > 0)
+		{
+			for (int32 i = 0; i < NumWindows; i++)
+			{
+				const FTimelineEventData::FEventWindow& Window = EventWindows[i];
+
+				FVector2D EventSize = EventBrush->GetImageSize();
+
+				const float XStart = RangeToScreen.InputToLocalX(Window.TimeStart);
+				const float XStartDiamond = XStart - EventSize.X / 2;
+				const float XEnd = RangeToScreen.InputToLocalX(Window.TimeEnd);
+				const float XEndDiamond = XEnd - EventSize.X / 2;
+				const float Y = (AllottedGeometry.Size.Y - EventSize.Y) / 2;
+
+				// window bar
+				FSlateDrawElement::MakeBox(OutDrawElements, LayerId++,
+					AllottedGeometry.ToPaintGeometry(FVector2D(XEnd - XStart, EventSize.Y - 2), FSlateLayoutTransform(FVector2D(XStart, Y + 1))), Brush, ESlateDrawEffect::None, Window.Color);
+
+				// key diamond at start
+				FSlateDrawElement::MakeBox(OutDrawElements, LayerId,
+					AllottedGeometry.ToPaintGeometry(EventSize, FSlateLayoutTransform(FVector2D(XStartDiamond, Y))), EventBrush, ESlateDrawEffect::None, Window.Color);
+
+				FSlateDrawElement::MakeBox(OutDrawElements, LayerId,
+					AllottedGeometry.ToPaintGeometry(EventSize, FSlateLayoutTransform(FVector2D(XStartDiamond, Y))), EventBorderBrush, ESlateDrawEffect::None, FLinearColor::Black);
+
+				// key diamond at end
+				FSlateDrawElement::MakeBox(OutDrawElements, LayerId,
+					AllottedGeometry.ToPaintGeometry(EventSize, FSlateLayoutTransform(FVector2D(XEndDiamond, Y))), EventBrush, ESlateDrawEffect::None, Window.Color);
+
+				FSlateDrawElement::MakeBox(OutDrawElements, LayerId,
+					AllottedGeometry.ToPaintGeometry(EventSize, FSlateLayoutTransform(FVector2D(XEndDiamond, Y))), EventBorderBrush, ESlateDrawEffect::None, FLinearColor::Black);
+			}
+		}
+
+		LayerId++;
+	}
+
+	return LayerId;
+}
+
+#undef LOCTEXT_NAMESPACE
diff --git a/Engine/Source/Editor/RewindDebuggerInterface/Public/SEventTimelineView.h b/Engine/Source/Editor/RewindDebuggerInterface/Public/SEventTimelineView.h
new file mode 100644
--- /dev/null
+++ b/Engine/Source/Editor/RewindDebuggerInterface/Public/SEventTimelineView.h
@@ -0,0 +1,118 @@
+// Copyright Epic Games, Inc. All Rights Reserved.
+
+#pragma once
+
+#if UE_ENABLE_INCLUDE_ORDER_DEPRECATED_IN_5_8
+#include "CoreMinimal.h"
+#endif // UE_ENABLE_INCLUDE_ORDER_DEPRECATED_IN_5_8
+#include "Widgets/SCompoundWidget.h"
+#include "Widgets/DeclarativeSyntaxSupport.h"
+
+#define UE_API REWINDDEBUGGERINTERFACE_API
+
+class FPaintArgs;
+class FSlateWindowElementList;
+class FTimeSliderController;
+
+class SEventTimelineView : public SCompoundWidget
+{
+public:
+	struct FTimelineEventData
+	{
+		struct FEventPoint
+		{
+			double Time;
+			FText Type;
+			FText Description;
+			FLinearColor Color;
+		};
+
+		struct FEventWindow
+		{
+			double TimeStart;
+			double TimeEnd;
+			FText Type;
+			FText Description;
+			FLinearColor Color;
+		};
+
+		TArray<FEventPoint> EventPoints;
+		TArray<FEventWindow> EventWindows;
+
+PRAGMA_DISABLE_DEPRECATION_WARNINGS
+#if WITH_EDITORONLY_DATA
+		struct UE_DEPRECATED(5.8, "Use FEventPoint instead") EventPoint
+		{
+			double Time;
+			FText Type;
+			FText Description;
+			FLinearColor Color;
+		};
+		struct UE_DEPRECATED(5.8, "Use FEventWindow instead") EventWindow
+		{
+			double TimeStart;
+			double TimeEnd;
+			FText Type;
+			FText Description;
+			FLinearColor Color;
+		};
+
+		UE_DEPRECATED(5.8, "Use EventPoints instead")
+		TArray<EventPoint> Points;
+		UE_DEPRECATED(5.8, "Use EventWindows instead")
+		TArray<EventWindow> Windows;
+#endif
+PRAGMA_ENABLE_DEPRECATION_WARNINGS
+	};
+
+	SLATE_BEGIN_ARGS(SEventTimelineView)
+		: _ViewRange(TRange<double>(0, 10))
+		, _DesiredSize(FVector2D(100.f, 20.f))
+	{
+	}
+	/** View time range */
+	SLATE_ATTRIBUTE(TRange<double>, ViewRange);
+
+	/** Data for events to render */
+	SLATE_ATTRIBUTE(TSharedPtr<FTimelineEventData>, EventData);
+
+	/** Desired widget size */
+	SLATE_ATTRIBUTE(FVector2D, DesiredSize);
+
+	SLATE_END_ARGS()
+
+
+	/**
+	 * Construct the widget
+	 *
+	 * @param InArgs A declaration from which to construct the widget
+	 */
+	UE_API void Construct(const FArguments& InArgs);
+
+protected:
+	//~ SWidget interface
+	UE_API virtual FVector2D ComputeDesiredSize(float) const override;
+	UE_API virtual int32 OnPaint(const FPaintArgs& Args
+		, const FGeometry& AllottedGeometry
+		, const FSlateRect& MyCullingRect
+		, FSlateWindowElementList& OutDrawElements
+		, int32 LayerId
+		, const FWidgetStyle& InWidgetStyle
+		, bool bParentEnabled) const override;
+
+	UE_API int32 PaintEvents(const FGeometry& AllottedGeometry
+		, const FSlateRect& MyCullingRect
+		, FSlateWindowElementList& OutDrawElements
+		, int32 LayerId
+		, const FWidgetStyle& InWidgetStyle
+		, bool bParentEnabled) const;
+
+	TAttribute<TRange<double>> ViewRange;
+	TAttribute<TSharedPtr<FTimelineEventData>> EventData;
+	TAttribute<FVector2D> DesiredSize;
+
+	const FSlateBrush* EventBrush = nullptr;
+	const FSlateBrush* EventBorderBrush = nullptr;
+};
+
+#undef UE_API
diff --git a/Engine/Source/Editor/RewindDebuggerInterface/RewindDebuggerInterface.Build.cs b/Engine/Source/Editor/RewindDebuggerInterface/RewindDebuggerInterface.Build.cs
--- a/Engine/Source/Editor/RewindDebuggerInterface/RewindDebuggerInterface.Build.cs
+++ b/Engine/Source/Editor/RewindDebuggerInterface/RewindDebuggerInterface.Build.cs
@@ -7,12 +7,18 @@ public class RewindDebuggerInterface : ModuleRules
 {
 	public RewindDebuggerInterface(ReadOnlyTargetRules Target) : base(Target)
 	{
-		PrivateDependencyModuleNames.AddRange(
+
+		PublicDependencyModuleNames.AddRange(
 			new string[] {
 				"Core",
 				"CoreUObject",
 				"SlateCore",
-				"TraceServices"
+			}
+		);
+		PrivateDependencyModuleNames.AddRange(
+			new string[] {
+				"Slate",
+				"ToolWidgets",
 			}
 		);
 	}
diff --git a/Engine/Source/Runtime/Mass/MassDeveloper/Private/Trace/MassTraceProvider.cpp b/Engine/Source/Runtime/Mass/MassDeveloper/Private/Trace/MassTraceProvider.cpp
--- a/Engine/Source/Runtime/Mass/MassDeveloper/Private/Trace/MassTraceProvider.cpp
+++ b/Engine/Source/Runtime/Mass/MassDeveloper/Private/Trace/MassTraceProvider.cpp
@@ -13,12 +13,13 @@ DEFINE_LOG_CATEGORY(LogMassTrace)
 
 namespace UE::Mass::Trace
 {
+static constexpr uint64 EntityEventsPageSize = 65536;
 
 thread_local TraceServices::FProviderLock::FThreadLocalState GMassInsightsProviderLockState;
 
 FMassTraceProvider::FMassTraceProvider(TraceServices::IAnalysisSession& InSession)
 	: Session(InSession)
-	, EntityEvents(InSession.GetLinearAllocator(), Constants::EntityEventsPageSize)
+	, EntityEvents(InSession.GetLinearAllocator(), EntityEventsPageSize)
 {
 }
 
diff --git a/Engine/Source/Runtime/Mass/MassDeveloper/Private/Trace/MassTraceProvider.h b/Engine/Source/Runtime/Mass/MassDeveloper/Private/Trace/MassTraceProvider.h
--- a/Engine/Source/Runtime/Mass/MassDeveloper/Private/Trace/MassTraceProvider.h
+++ b/Engine/Source/Runtime/Mass/MassDeveloper/Private/Trace/MassTraceProvider.h
@@ -27,12 +27,8 @@ class FMassTraceProvider
 	: public IMassTraceProvider
 	, public IEditableMassTraceProvider
 {
-	struct Constants
-	{
-		static constexpr uint64 EntityEventsPageSize = 65536;
-	};
 public:
-	explicit FMassTraceProvider(TraceServices::IAnalysisSession& Session);
+	explicit FMassTraceProvider(TraceServices::IAnalysisSession& InSession);
 	virtual ~FMassTraceProvider() override = default;
 
 	//~ Read operations
@@ -101,13 +97,11 @@ public:
 
 	virtual void AppendRegionBegin(const TCHAR* Name, uint64 ID, double Time) override;
 	virtual void AppendRegionBegin(const TCHAR* Name, double Time) override;
-	virtual void AppendRegionEnd(const TCHAR* ID, double Time) override;
+	virtual void AppendRegionEnd(const TCHAR* Name, double Time) override;
 	virtual void AppendRegionEnd(const uint64 ID, double Time) override;
 
 	virtual void OnAnalysisSessionEnded() override;
 
-	//////////////////////////////////////////////////
-
 private:
 	/** Update the depth member of a region to allow overlapping regions to be displayed on separate lanes. */
 	int32 CalculateRegionDepth(const FRegion& Item) const;
