[MassDebug]
Fix MassEntity RewindDebugger timeline misalignment during simulation pauses

Emit RecordingTime (FObjectTrace::GetWorldElapsedTime) directly from Mass trace events, matching the pattern used by Animation, StateTree, and other trace systems.
FMassTrace::SetRecordingTime is called at each phase start from FMassProcessingPhaseManager::OnPhaseStart via UE_TRACE_MASS_SET_RECORDING_TIME.

Additional changes:
- FEntityEventRecord carries both Time (ElapsedTime) and ProfileTime
- Details view diffs archetype before the current frame batch vs current,with substep attribution for merged same-frame changes
- Entity track supports stepping to prev/next event via GetStepFrameTimeInternal

Thread safety: All writes (SetRecordingTime in OnPhaseStart) and reads (EntityCreated/Moved/Destroyed) occur on the game thread
-> entity structural modifications enforce game-thread-only access via CHECK_SYNC_API(); if parallel entity mutations are ever introduced, this should be revisited.

Backward compatibility: RecordingTime is a new optional trace event field.
Old traces without it return the sentinel default (-1.0) and the analyzer falls back to ProfileTime, preserving the old behavior (not pause-aligned but functional).
Old analyzers reading new traces simply ignore the unknown field.

#Synced-CL 52444545

diff --git a/Engine/Source/Editor/MassEntityEditor/Private/Trace/MassEntityEventCache.cpp b/Engine/Source/Editor/MassEntityEditor/Private/Trace/MassEntityEventCache.cpp
--- a/Engine/Source/Editor/MassEntityEditor/Private/Trace/MassEntityEventCache.cpp
+++ b/Engine/Source/Editor/MassEntityEditor/Private/Trace/MassEntityEventCache.cpp
@@ -46,15 +46,18 @@ void FEntityEventCache::Update(
 
 	ConvertTimesToRecordingTime(Session);
 
-	// Sort each entity's events by time so consumers can rely on ascending order
-	// (e.g. early-exit in ResolveArchetypesAtScrubTime).
+	// Sort each entity's events by ProfileTime so consumers can rely on ascending
+	// chronological order.  We use ProfileTime rather than Time (ElapsedTime) because
+	// ElapsedTime freezes during pauses — events in different frames can share the
+	// same ElapsedTime value, and sorting by it could interleave frames and break
+	// the FrameIndex-based batching in ResolveArchetypesAtScrubTime.
 	{
 		FWriteScopeLock WriteLock(CacheLock);
 		for (TPair<uint64, TArray<FCachedEntityEventRecord>>& Pair : EntityEvents)
 		{
 			Pair.Value.Sort([](const FCachedEntityEventRecord& A, const FCachedEntityEventRecord& B)
 			{
-				return A.Time < B.Time;
+				return A.ProfileTime < B.ProfileTime;
 			});
 		}
 	}
@@ -65,16 +68,23 @@ void FEntityEventCache::ConvertTimesToRecordingTime(const TraceServices::IAnalys
 	TraceServices::FAnalysisSessionReadScope SessionReadScope(Session);
 	const TraceServices::IFrameProvider& FrameProvider = TraceServices::ReadFrameProvider(Session);
 
-	// Pass 1: resolve frame indices while Event.Time is still in ProfileTime space.
+	// Resolve frame indices from ProfileTime.
 	for (TPair<uint64, TArray<FCachedEntityEventRecord>>& Pair : EntityEvents)
 	{
 		for (FCachedEntityEventRecord& Event : Pair.Value)
 		{
-			Event.FrameIndex = FrameProvider.GetFrameNumberForTimestamp(TraceFrameType_Game, Event.Time);
+			Event.FrameIndex = FrameProvider.GetFrameNumberForTimestamp(TraceFrameType_Game, Event.ProfileTime);
 		}
 	}
 
-	// Pass 2: normalize event times relative to the first frame's start time.
+	// For new traces, Event.Time is already in ElapsedTime (recording time) space —
+	// set at trace emission via FMassTrace::SetRecordingTime / FObjectTrace::GetWorldElapsedTime.
+	//
+	// For old traces that lack the RecordingTime field, the analyzer falls back to
+	// ProfileTime for Event.Time (Time == ProfileTime). These are raw absolute values
+	// that are incomparable to the RewindDebugger's ElapsedTime-based scrub time.
+	// Normalize them relative to the first frame's start time so they become
+	// approximately comparable (not pause-aligned, but functional).
 	const uint64 FrameCount = FrameProvider.GetFrameCount(TraceFrameType_Game);
 	if (FrameCount == 0)
 	{
@@ -92,7 +102,13 @@ void FEntityEventCache::ConvertTimesToRecordingTime(const TraceServices::IAnalys
 	{
 		for (FCachedEntityEventRecord& Event : Pair.Value)
 		{
-			Event.Time -= ProfileStart;
+			// Events from old traces have Time == ProfileTime (analyzer fallback).
+			// Normalize these to be relative to trace start. Events with proper
+			// RecordingTime (Time != ProfileTime) are left untouched.
+			if (Event.Time == Event.ProfileTime)
+			{
+				Event.Time -= ProfileStart;
+			}
 		}
 	}
 }
diff --git a/Engine/Source/Editor/MassEntityEditor/Private/Trace/MassEntityEventCache.h b/Engine/Source/Editor/MassEntityEditor/Private/Trace/MassEntityEventCache.h
--- a/Engine/Source/Editor/MassEntityEditor/Private/Trace/MassEntityEventCache.h
+++ b/Engine/Source/Editor/MassEntityEditor/Private/Trace/MassEntityEventCache.h
@@ -29,8 +29,7 @@ public:
 		uint32 FrameIndex = 0;
 	};
 
-	/** Rebuild indices from the provider if GetUpdateCounter() has changed.
-	 * Converts event times from ProfileTime to RecordingTime using the IFrameProvider. */
+	/** Rebuild indices from the provider if GetUpdateCounter() has changed. */
 	void Update(const IMassTraceProvider& Provider,
 		const TraceServices::IAnalysisSession& Session);
 
@@ -55,10 +54,13 @@ public:
 	}
 
 private:
-	/** Convert event times from ProfileTime to ElapsedTime (recording time). */
+	/** Resolve frame indices from ProfileTime and ensure Event.Time is in ElapsedTime
+	 * (recording time) space. New traces already carry correct ElapsedTime from emission;
+	 * old traces lacking the RecordingTime field have Time == ProfileTime and are normalized
+	 * relative to the first frame's start time as a backward-compatibility fallback. */
 	void ConvertTimesToRecordingTime(const TraceServices::IAnalysisSession& Session);
 
-	/** Entity ID -> array of events, sorted by ascending time after Update() */
+	/** Entity ID -> array of events, sorted by ascending ProfileTime after Update() */
 	TMap<uint64, TArray<FCachedEntityEventRecord>> EntityEvents;
 
 	/** Archetype ID -> set of entity IDs that reference it */
diff --git a/Engine/Source/Editor/MassEntityEditor/Private/Trace/MassEntityTrack.cpp b/Engine/Source/Editor/MassEntityEditor/Private/Trace/MassEntityTrack.cpp
--- a/Engine/Source/Editor/MassEntityEditor/Private/Trace/MassEntityTrack.cpp
+++ b/Engine/Source/Editor/MassEntityEditor/Private/Trace/MassEntityTrack.cpp
@@ -296,6 +296,60 @@ bool FEntityTrack::HasDebugDataInternal() const
 	return Cache->GetEntityEvents(EntityId).Num() > 0;
 }
 
+FText FEntityTrack::GetStepCommandTooltipInternal(const RewindDebugger::EStepMode StepMode) const
+{
+	using RewindDebugger::EStepMode;
+	switch (StepMode)
+	{
+	case EStepMode::Forward:
+		return LOCTEXT("StepForwardTooltip", "Step to next event in Mass Entity track");
+	case EStepMode::Backward:
+		return LOCTEXT("StepBackwardTooltip", "Step to previous event in Mass Entity track");
+	default:
+		ensureMsgf(false, TEXT("Unhandled step mode value: %d"), StepMode);
+		break;
+	}
+	return {};
+}
+
+TOptional<double> FEntityTrack::GetStepFrameTimeInternal(const RewindDebugger::EStepMode StepMode, const FScrubTimeInformation& CurrentScrubTime) const
+{
+	const TConstArrayView<FEntityEventCache::FCachedEntityEventRecord> Events = Cache->GetEntityEvents(EntityId);
+	if (Events.Num() == 0)
+	{
+		return {};
+	}
+
+	const double ScrubTime = CurrentScrubTime.ElapsedTime;
+
+	using RewindDebugger::EStepMode;
+	if (StepMode == EStepMode::Forward)
+	{
+		// Find the first event strictly after the current scrub time
+		for (const FEntityEventCache::FCachedEntityEventRecord& Event : Events)
+		{
+			if (Event.Time > ScrubTime + UE_DOUBLE_SMALL_NUMBER)
+			{
+				return Event.Time;
+			}
+		}
+	}
+	else
+	{
+		// Find the last event strictly before the current scrub time
+		for (int32 Index = Events.Num() - 1; Index >= 0; --Index)
+		{
+			if (Events[Index].Time < ScrubTime - UE_DOUBLE_SMALL_NUMBER)
+			{
+				return Events[Index].Time;
+			}
+		}
+	}
+
+	// No more events in this direction — preserve current position
+	return ScrubTime;
+}
+
 } // namespace UE::Mass::Trace
 
 #undef LOCTEXT_NAMESPACE
diff --git a/Engine/Source/Editor/MassEntityEditor/Private/Trace/MassEntityTrack.h b/Engine/Source/Editor/MassEntityEditor/Private/Trace/MassEntityTrack.h
--- a/Engine/Source/Editor/MassEntityEditor/Private/Trace/MassEntityTrack.h
+++ b/Engine/Source/Editor/MassEntityEditor/Private/Trace/MassEntityTrack.h
@@ -39,6 +39,8 @@ private:
 		return EntityId;
 	}
 	virtual bool HasDebugDataInternal() const override;
+	virtual FText GetStepCommandTooltipInternal(RewindDebugger::EStepMode StepMode) const override;
+	virtual TOptional<double> GetStepFrameTimeInternal(RewindDebugger::EStepMode StepMode, const FScrubTimeInformation& CurrentScrubTime) const override;
 
 	TSharedPtr<SEventTimelineView::FTimelineEventData> GetEventData() const;
 
diff --git a/Engine/Source/Editor/MassEntityEditor/Private/Trace/SMassEntityDetailsView.cpp b/Engine/Source/Editor/MassEntityEditor/Private/Trace/SMassEntityDetailsView.cpp
--- a/Engine/Source/Editor/MassEntityEditor/Private/Trace/SMassEntityDetailsView.cpp
+++ b/Engine/Source/Editor/MassEntityEditor/Private/Trace/SMassEntityDetailsView.cpp
@@ -225,20 +225,17 @@ void SMassEntityDetailsView::Tick(const FGeometry& AllottedGeometry, const doubl
 	const double ScrubTime = IRewindDebugger::Instance()->GetScrubTime();
 	SetScrubTime(ScrubTime);
 
-	uint64 CurrentArchetype = 0;
-	uint64 PreviousArchetype = 0;
-	int32 MergeCount = 0;
-	TArray<uint64> BatchArchetypeIds;
-	ResolveArchetypesAtScrubTime(CurrentArchetype, PreviousArchetype, MergeCount, BatchArchetypeIds);
-
-	if (CurrentArchetype != LastDisplayedArchetypeId ||
-		PreviousArchetype != LastDisplayedPreviousArchetypeId ||
-		MergeCount != LastDisplayedMergeCount)
+	const FArchetypeResolution Resolution = ResolveArchetypesAtScrubTime();
+
+	if (Resolution.CurrentArchetype != LastDisplayedArchetypeId ||
+		Resolution.PreviousArchetype != LastDisplayedPreviousArchetypeId ||
+		Resolution.MergeCount != LastDisplayedMergeCount)
 	{
-		LastDisplayedArchetypeId = CurrentArchetype;
-		LastDisplayedPreviousArchetypeId = PreviousArchetype;
-		LastDisplayedMergeCount = MergeCount;
-		LastBatchArchetypeIds = MoveTemp(BatchArchetypeIds);
+		LastDisplayedArchetypeId = Resolution.CurrentArchetype;
+		LastDisplayedPreviousArchetypeId = Resolution.PreviousArchetype;
+		LastDisplayedArchetypeBeforeBatch = Resolution.ArchetypeBeforeBatch;
+		LastDisplayedMergeCount = Resolution.MergeCount;
+		LastBatchArchetypeIds = Resolution.BatchArchetypeIds;
 		RebuildFragmentList();
 	}
 }
@@ -250,24 +247,19 @@ TSharedRef<ITableRow> SMassEntityDetailsView::OnGenerateRow(
 	return SNew(SMassArchetypeDiffRow, OwnerTable).Entry(Entry);
 }
 
-void SMassEntityDetailsView::ResolveArchetypesAtScrubTime(uint64& OutCurrentArchetype, uint64& OutPreviousArchetype,
-	int32& OutMergeCount, TArray<uint64>& OutBatchArchetypeIds) const
+SMassEntityDetailsView::FArchetypeResolution SMassEntityDetailsView::ResolveArchetypesAtScrubTime() const
 {
 	using namespace UE::Mass::Trace;
 
-	OutCurrentArchetype = 0;
-	OutPreviousArchetype = 0;
-	OutMergeCount = 0;
-	OutBatchArchetypeIds.Reset();
+	FArchetypeResolution Result;
 
 	if (!Cache)
 	{
-		return;
+		return Result;
 	}
 
 	TConstArrayView<FEntityEventCache::FCachedEntityEventRecord> Events = Cache->GetEntityEvents(EntityId);
 	uint32 LastFrameIndex = MAX_uint32;
-	uint64 ArchetypeBeforeCurrentBatch = 0;
 	int32 CurrentBatchChangeCount = 0;
 	TArray<uint64> CurrentBatchArchetypeIds;
 
@@ -278,29 +270,34 @@ void SMassEntityDetailsView::ResolveArchetypesAtScrubTime(uint64& OutCurrentArch
 			break;
 		}
 
+		// Track same-frame batching using FrameIndex (resolved from ProfileTime)
 		if (Record.FrameIndex != LastFrameIndex)
 		{
-			ArchetypeBeforeCurrentBatch = OutCurrentArchetype;
+			Result.ArchetypeBeforeBatch = Result.CurrentArchetype;
 			LastFrameIndex = Record.FrameIndex;
 			CurrentBatchChangeCount = 0;
 			CurrentBatchArchetypeIds.Reset();
 		}
 
+		// Previous archetype is always the archetype from the prior event
 		switch (Record.Operation)
 		{
 		case EEntityEventType::Created:
-			ArchetypeBeforeCurrentBatch = 0;
-			OutCurrentArchetype = Record.ArchetypeID;
+			Result.ArchetypeBeforeBatch = 0;
+			Result.PreviousArchetype = 0;
+			Result.CurrentArchetype = Record.ArchetypeID;
 			CurrentBatchArchetypeIds.Add(Record.ArchetypeID);
 			CurrentBatchChangeCount++;
 			break;
 		case EEntityEventType::ArchetypeChange:
-			OutCurrentArchetype = Record.ArchetypeID;
+			Result.PreviousArchetype = Result.CurrentArchetype;
+			Result.CurrentArchetype = Record.ArchetypeID;
 			CurrentBatchArchetypeIds.Add(Record.ArchetypeID);
 			CurrentBatchChangeCount++;
 			break;
 		case EEntityEventType::Destroyed:
-			OutCurrentArchetype = 0;
+			Result.PreviousArchetype = Result.CurrentArchetype;
+			Result.CurrentArchetype = 0;
 			CurrentBatchChangeCount++;
 			break;
 		default:
@@ -308,9 +305,21 @@ void SMassEntityDetailsView::ResolveArchetypesAtScrubTime(uint64& OutCurrentArch
 		}
 	}
 
-	OutPreviousArchetype = ArchetypeBeforeCurrentBatch;
-	OutMergeCount = CurrentBatchChangeCount;
-	OutBatchArchetypeIds = MoveTemp(CurrentBatchArchetypeIds);
+	// When the batch starts with entity creation (ArchetypeBeforeBatch == 0)
+	// but has subsequent changes, promote the Created archetype from the batch
+	// into the diff base so the details view shows what changed since creation.
+	// Remove it from the batch list so the substep walk doesn't compare it
+	// against itself at step 0.
+	if (Result.ArchetypeBeforeBatch == 0 && CurrentBatchChangeCount > 1 && CurrentBatchArchetypeIds.Num() > 0)
+	{
+		Result.ArchetypeBeforeBatch = CurrentBatchArchetypeIds[0];
+		CurrentBatchArchetypeIds.RemoveAt(0);
+		CurrentBatchChangeCount--;
+	}
+
+	Result.MergeCount = CurrentBatchChangeCount;
+	Result.BatchArchetypeIds = MoveTemp(CurrentBatchArchetypeIds);
+	return Result;
 }
 
 void SMassEntityDetailsView::RebuildFragmentList()
@@ -320,7 +329,10 @@ void SMassEntityDetailsView::RebuildFragmentList()
 	FragmentEntries.Reset();
 
 	const uint64 ArchetypeAtScrubTime = LastDisplayedArchetypeId;
-	const uint64 PreviousArchetypeId = LastDisplayedPreviousArchetypeId;
+	// Use the archetype before the current frame's batch for the diff.
+	// For a single change this equals the per-event previous; for merged
+	// changes it captures the full delta across all steps in the frame.
+	const uint64 PreviousArchetypeId = LastDisplayedArchetypeBeforeBatch;
 
 	// Update merge info text
 	if (LastDisplayedMergeCount > 1 && PreviousArchetypeId != 0)
@@ -398,15 +410,30 @@ void SMassEntityDetailsView::RebuildFragmentList()
 		}
 	}
 
-	// Step attribution for merged batches
+	// Step attribution for merged batches — walk through each intermediate archetype
+	// in the same-frame batch to determine at which step each fragment was added/removed.
 	TMap<uint64, int32> AddedAtStep;
 	TMap<uint64, int32> RemovedAtStep;
 	TSet<uint64> ReaddedIds;
 
-	if (LastDisplayedMergeCount > 1 && PreviousArchetype)
+	if (LastDisplayedMergeCount > 1 && LastDisplayedArchetypeBeforeBatch != 0)
 	{
 		TraceServices::FProviderReadScopeLock ReadLock(MassTraceProvider);
-		TSet<uint64> StepFragmentIds = PreviousFragmentIds;
+
+		// Start the substep walk from the archetype the entity was in before the
+		// first event in this frame's batch — not from the previous event's archetype,
+		// which may be from an entirely different point in the timeline.
+		TSet<uint64> StepFragmentIds;
+		if (const FArchetypeInfo* BatchStartArchetype = MassTraceProvider.FindArchetypeById(LastDisplayedArchetypeBeforeBatch))
+		{
+			for (const FFragmentInfo* Fragment : BatchStartArchetype->Fragments)
+			{
+				if (Fragment)
+				{
+					StepFragmentIds.Add(Fragment->Id);
+				}
+			}
+		}
 
 		for (int32 StepIdx = 0; StepIdx < LastBatchArchetypeIds.Num(); ++StepIdx)
 		{
@@ -459,7 +486,7 @@ void SMassEntityDetailsView::RebuildFragmentList()
 		}
 	}
 
-	// Build entries
+	// Build entries — diff between previous event's archetype and current
 	for (const FFragmentInfo* Fragment : CurrentArchetype->Fragments)
 	{
 		if (!Fragment)
diff --git a/Engine/Source/Editor/MassEntityEditor/Private/Trace/SMassEntityDetailsView.h b/Engine/Source/Editor/MassEntityEditor/Private/Trace/SMassEntityDetailsView.h
--- a/Engine/Source/Editor/MassEntityEditor/Private/Trace/SMassEntityDetailsView.h
+++ b/Engine/Source/Editor/MassEntityEditor/Private/Trace/SMassEntityDetailsView.h
@@ -51,9 +51,18 @@ public:
 private:
 	virtual void Tick(const FGeometry& AllottedGeometry, const double InCurrentTime, const float InDeltaTime) override;
 
+	struct FArchetypeResolution
+	{
+		uint64 CurrentArchetype = 0;
+		uint64 PreviousArchetype = 0;
+		/** Archetype the entity was in before the first event in the current frame's batch. */
+		uint64 ArchetypeBeforeBatch = 0;
+		int32 MergeCount = 0;
+		TArray<uint64> BatchArchetypeIds;
+	};
+
 	void RebuildFragmentList();
-	void ResolveArchetypesAtScrubTime(uint64& OutCurrentArchetype, uint64& OutPreviousArchetype,
-		int32& OutMergeCount, TArray<uint64>& OutBatchArchetypeIds) const;
+	FArchetypeResolution ResolveArchetypesAtScrubTime() const;
 
 	TSharedRef<ITableRow> OnGenerateRow(TSharedPtr<FFragmentEntry> Entry, const TSharedRef<STableViewBase>& OwnerTable);
 
@@ -62,6 +71,7 @@ private:
 	double CurrentScrubTime = 0.0;
 	uint64 LastDisplayedArchetypeId = 0;
 	uint64 LastDisplayedPreviousArchetypeId = 0;
+	uint64 LastDisplayedArchetypeBeforeBatch = 0;
 	int32 LastDisplayedMergeCount = 0;
 	TArray<uint64> LastBatchArchetypeIds;
 
diff --git a/Engine/Source/Runtime/Mass/MassDeveloper/Private/Trace/MassTraceAnalyzer.cpp b/Engine/Source/Runtime/Mass/MassDeveloper/Private/Trace/MassTraceAnalyzer.cpp
--- a/Engine/Source/Runtime/Mass/MassDeveloper/Private/Trace/MassTraceAnalyzer.cpp
+++ b/Engine/Source/Runtime/Mass/MassDeveloper/Private/Trace/MassTraceAnalyzer.cpp
@@ -104,22 +104,31 @@ bool FMassTraceAnalyzer::OnEvent(uint16 RouteId, EStyle Style, const FOnEventCon
 	case RouteId_MassBulkAddEntity:
 	{
 		const uint64 Cycle = EventData.GetValue<uint64>("Cycle");
+		const double ProfileTime = Context.EventTime.AsSeconds(Cycle);
+		// RecordingTime field was added later; old traces return -1.0 (missing field sentinel).
+		// Fall back to ProfileTime which gives the old behavior (not pause-aligned but functional).
+		const double RecordingTime = EventData.GetValue<double>("RecordingTime", -1.0);
+		const double EffectiveTime = RecordingTime >= 0.0 ? RecordingTime : ProfileTime;
 		const TArrayView<const uint64> Entities = EventData.GetArrayView<uint64>("Entities");
 		const TArrayView<const uint64> ArchetypeIDs = EventData.GetArrayView<uint64>("ArchetypeIDs");
 
 		TraceServices::FProviderEditScopeLock EditLock(MassTraceProvider.Get());
-		MassTraceProvider.Get().BulkAddEntity(Context.EventTime.AsSeconds(Cycle), Entities, ArchetypeIDs);
+		MassTraceProvider.Get().BulkAddEntity(EffectiveTime, ProfileTime, Entities, ArchetypeIDs);
 		break;
 	}
 	case RouteId_MassEntityMoved:
 	{
 		const uint64 Cycle = EventData.GetValue<uint64>("Cycle");
+		const double ProfileTime = Context.EventTime.AsSeconds(Cycle);
+		const double RecordingTime = EventData.GetValue<double>("RecordingTime", -1.0);
+		const double EffectiveTime = RecordingTime >= 0.0 ? RecordingTime : ProfileTime;
 		const uint64 Entity = EventData.GetValue<uint64>("Entity");
 		const uint64 Archetype = EventData.GetValue<uint64>("NewArchetypeID");
 
 		TraceServices::FProviderEditScopeLock EditLock(MassTraceProvider.Get());
 		MassTraceProvider.Get().BulkMoveEntity(
-			Context.EventTime.AsSeconds(Cycle),
+			EffectiveTime,
+			ProfileTime,
 			MakeConstArrayView(&Entity, 1),
 			MakeConstArrayView(&Archetype, 1));
 		break;
@@ -127,11 +136,15 @@ bool FMassTraceAnalyzer::OnEvent(uint16 RouteId, EStyle Style, const FOnEventCon
 	case RouteId_MassBulkEntityDestroyed:
 	{
 		const uint64 Cycle = EventData.GetValue<uint64>("Cycle");
+		const double ProfileTime = Context.EventTime.AsSeconds(Cycle);
+		const double RecordingTime = EventData.GetValue<double>("RecordingTime", -1.0);
+		const double EffectiveTime = RecordingTime >= 0.0 ? RecordingTime : ProfileTime;
 		const TArrayView<const uint64> Entities = EventData.GetArrayView<uint64>("Entities");
 
 		TraceServices::FProviderEditScopeLock EditLock(MassTraceProvider.Get());
 		MassTraceProvider.Get().BulkDestroyEntity(
-			Context.EventTime.AsSeconds(Cycle),
+			EffectiveTime,
+			ProfileTime,
 			Entities);
 		break;
 	}
diff --git a/Engine/Source/Runtime/Mass/MassDeveloper/Private/Trace/MassTraceProvider.cpp b/Engine/Source/Runtime/Mass/MassDeveloper/Private/Trace/MassTraceProvider.cpp
--- a/Engine/Source/Runtime/Mass/MassDeveloper/Private/Trace/MassTraceProvider.cpp
+++ b/Engine/Source/Runtime/Mass/MassDeveloper/Private/Trace/MassTraceProvider.cpp
@@ -158,12 +158,13 @@ void FMassTraceProvider::AddArchetype(const FArchetypeInfo& ArchetypeInfo)
 	}
 }
 
-void FMassTraceProvider::BulkAddEntity(const double Time, const TConstArrayView<uint64> Entities, const TConstArrayView<uint64> ArchetypeIDs)
+void FMassTraceProvider::BulkAddEntity(const double RecordingTime, const double ProfileTime, const TConstArrayView<uint64> Entities, const TConstArrayView<uint64> ArchetypeIDs)
 {
 	EditAccessCheck();
 
 	FEntityEventRecord Row;
-	Row.Time = Time;
+	Row.Time = RecordingTime;
+	Row.ProfileTime = ProfileTime;
 	Row.Operation = EEntityEventType::Created;
 
 	for (int32 Index = 0, End = Entities.Num(); Index < End; Index++)
@@ -174,12 +175,13 @@ void FMassTraceProvider::BulkAddEntity(const double Time, const TConstArrayView<
 	}
 }
 
-void FMassTraceProvider::BulkMoveEntity(const double Time, const TConstArrayView<uint64> Entities, const TConstArrayView<uint64> ArchetypeIDs)
+void FMassTraceProvider::BulkMoveEntity(const double RecordingTime, const double ProfileTime, const TConstArrayView<uint64> Entities, const TConstArrayView<uint64> ArchetypeIDs)
 {
 	EditAccessCheck();
 
 	FEntityEventRecord Row;
-	Row.Time = Time;
+	Row.Time = RecordingTime;
+	Row.ProfileTime = ProfileTime;
 	Row.Operation = EEntityEventType::ArchetypeChange;
 
 	for (int32 Index = 0, End = Entities.Num(); Index < End; Index++)
@@ -190,12 +192,13 @@ void FMassTraceProvider::BulkMoveEntity(const double Time, const TConstArrayView
 	}
 }
 
-void FMassTraceProvider::BulkDestroyEntity(const double Time, const TConstArrayView<uint64> Entities)
+void FMassTraceProvider::BulkDestroyEntity(const double RecordingTime, const double ProfileTime, const TConstArrayView<uint64> Entities)
 {
 	EditAccessCheck();
 
 	FEntityEventRecord Row;
-	Row.Time = Time;
+	Row.Time = RecordingTime;
+	Row.ProfileTime = ProfileTime;
 	Row.Operation = EEntityEventType::Destroyed;
 	Row.ArchetypeID = 0;
 
diff --git a/Engine/Source/Runtime/Mass/MassDeveloper/Private/Trace/MassTraceProvider.h b/Engine/Source/Runtime/Mass/MassDeveloper/Private/Trace/MassTraceProvider.h
--- a/Engine/Source/Runtime/Mass/MassDeveloper/Private/Trace/MassTraceProvider.h
+++ b/Engine/Source/Runtime/Mass/MassDeveloper/Private/Trace/MassTraceProvider.h
@@ -91,9 +91,9 @@ public:
 	virtual void AddFragment(const FFragmentInfo& FragmentInfo) override;
 	virtual void AddArchetype(const FArchetypeInfo& ArchetypeInfo) override;
 
-	virtual void BulkAddEntity(double Time, TConstArrayView<uint64> Entities, TConstArrayView<uint64> ArchetypeIDs) override;
-	virtual void BulkMoveEntity(double Time, TConstArrayView<uint64> Entities, TConstArrayView<uint64> ArchetypeIDs) override;
-	virtual void BulkDestroyEntity(double Time, TConstArrayView<uint64> Entities) override;
+	virtual void BulkAddEntity(double RecordingTime, double ProfileTime, TConstArrayView<uint64> Entities, TConstArrayView<uint64> ArchetypeIDs) override;
+	virtual void BulkMoveEntity(double RecordingTime, double ProfileTime, TConstArrayView<uint64> Entities, TConstArrayView<uint64> ArchetypeIDs) override;
+	virtual void BulkDestroyEntity(double RecordingTime, double ProfileTime, TConstArrayView<uint64> Entities) override;
 
 	virtual void AppendRegionBegin(const TCHAR* Name, uint64 ID, double Time) override;
 	virtual void AppendRegionBegin(const TCHAR* Name, double Time) override;
diff --git a/Engine/Source/Runtime/Mass/MassDeveloper/Public/Trace/MassTraceTypes.h b/Engine/Source/Runtime/Mass/MassDeveloper/Public/Trace/MassTraceTypes.h
--- a/Engine/Source/Runtime/Mass/MassDeveloper/Public/Trace/MassTraceTypes.h
+++ b/Engine/Source/Runtime/Mass/MassDeveloper/Public/Trace/MassTraceTypes.h
@@ -96,7 +96,12 @@ enum class EEntityEventType : uint8
 
 struct FEntityEventRecord
 {
+	/** RecordingTime (ElapsedTime) — game time that freezes during pauses.
+	 *  This is the time domain the RewindDebugger scrub/view range operates in. */
 	double Time = 0;
+	/** ProfileTime — wall-clock time that keeps advancing during pauses.
+	 *  Used for frame index resolution via IFrameProvider. */
+	double ProfileTime = 0;
 	uint64 Entity = INDEX_NONE;
 	uint64 ArchetypeID = INDEX_NONE;
 	EEntityEventType Operation = EEntityEventType::Unknown;
@@ -177,9 +182,9 @@ public:
 	virtual void AddFragment(const FFragmentInfo& FragmentInfo) = 0;
 	virtual void AddArchetype(const FArchetypeInfo& ArchetypeInfo) = 0;
 
-	virtual void BulkAddEntity(double Time, TConstArrayView<uint64> Entities, TConstArrayView<uint64> ArchetypeIDs) = 0;
-	virtual void BulkMoveEntity(double Time, TConstArrayView<uint64> Entities, TConstArrayView<uint64> ArchetypeIDs) = 0;
-	virtual void BulkDestroyEntity(double Time, TConstArrayView<uint64> Entities) = 0;
+	virtual void BulkAddEntity(double RecordingTime, double ProfileTime, TConstArrayView<uint64> Entities, TConstArrayView<uint64> ArchetypeIDs) = 0;
+	virtual void BulkMoveEntity(double RecordingTime, double ProfileTime, TConstArrayView<uint64> Entities, TConstArrayView<uint64> ArchetypeIDs) = 0;
+	virtual void BulkDestroyEntity(double RecordingTime, double ProfileTime, TConstArrayView<uint64> Entities) = 0;
 
 	/**
 	 * Append a new begin event of a region from the trace session.
diff --git a/Engine/Source/Runtime/MassEntity/Private/MassEntityTrace.cpp b/Engine/Source/Runtime/MassEntity/Private/MassEntityTrace.cpp
--- a/Engine/Source/Runtime/MassEntity/Private/MassEntityTrace.cpp
+++ b/Engine/Source/Runtime/MassEntity/Private/MassEntityTrace.cpp
@@ -14,6 +14,11 @@
 #include "Trace/Trace.inl"
 #include "TraceFilter.h"
 
+namespace UE::Mass::Trace
+{
+static double GMassRecordingTime = -1.0;
+}
+
 UE_TRACE_CHANNEL_DEFINE(MassChannel, "Mass Entity system archetypes, queries, and chunk execution.");
 
 UE_TRACE_EVENT_BEGIN(MassTrace, MassPhaseBegin)
@@ -65,17 +70,20 @@ UE_TRACE_EVENT_END()
 
 UE_TRACE_EVENT_BEGIN(MassTrace, MassBulkAddEntity)
 	UE_TRACE_EVENT_FIELD(uint64, Cycle)
+	UE_TRACE_EVENT_FIELD(double, RecordingTime)
 	UE_TRACE_EVENT_FIELD(uint64[], Entities)
 	UE_TRACE_EVENT_FIELD(uint64[], ArchetypeIDs)
 UE_TRACE_EVENT_END()
 
 UE_TRACE_EVENT_BEGIN(MassTrace, MassBulkEntityDestroyed)
 	UE_TRACE_EVENT_FIELD(uint64, Cycle)
+	UE_TRACE_EVENT_FIELD(double, RecordingTime)
 	UE_TRACE_EVENT_FIELD(uint64[], Entities)
 UE_TRACE_EVENT_END()
 
 UE_TRACE_EVENT_BEGIN(MassTrace, MassEntityMoved)
 	UE_TRACE_EVENT_FIELD(uint64, Cycle)
+	UE_TRACE_EVENT_FIELD(double, RecordingTime)
 	UE_TRACE_EVENT_FIELD(uint64, Entity)
 	UE_TRACE_EVENT_FIELD(uint64, NewArchetypeID)
 UE_TRACE_EVENT_END()
@@ -329,6 +337,7 @@ void FMassTrace::EntityCreated(FMassEntityHandle Entity, const FMassArchetypeDat
 		const uint64 ArchetypeID = FMassDebugger::GetArchetypeTraceID(Archetype);
 		UE_TRACE_LOG(MassTrace, MassBulkAddEntity, MassChannel)
 			<< MassBulkAddEntity.Cycle(Cycle)
+			<< MassBulkAddEntity.RecordingTime(UE::Mass::Trace::GMassRecordingTime)
 			<< MassBulkAddEntity.Entities(&EntityAsU64, 1)
 			<< MassBulkAddEntity.ArchetypeIDs(&ArchetypeID, 1);
 	}
@@ -342,6 +351,7 @@ void FMassTrace::EntityMoved(FMassEntityHandle Entity, const FMassArchetypeData&
 		const uint64 EntityAsU64 = Entity.AsNumber();
 		UE_TRACE_LOG(MassTrace, MassEntityMoved, MassChannel)
 			<< MassEntityMoved.Cycle(Cycle)
+			<< MassEntityMoved.RecordingTime(UE::Mass::Trace::GMassRecordingTime)
 			<< MassEntityMoved.Entity(EntityAsU64)
 			<< MassEntityMoved.NewArchetypeID(FMassDebugger::GetArchetypeTraceID(NewArchetype));
 	}
@@ -355,6 +365,7 @@ void FMassTrace::EntityDestroyed(FMassEntityHandle Entity)
 		const uint64 EntityAsU64 = Entity.AsNumber();
 		UE_TRACE_LOG(MassTrace, MassBulkEntityDestroyed, MassChannel)
 			<< MassBulkEntityDestroyed.Cycle(Cycle)
+			<< MassBulkEntityDestroyed.RecordingTime(UE::Mass::Trace::GMassRecordingTime)
 			<< MassBulkEntityDestroyed.Entities(&EntityAsU64, 1);
 	}
 }
@@ -370,10 +381,16 @@ void FMassTrace::EntitiesDestroyed(TConstArrayView<FMassEntityHandle> Entities)
 			Entities.Num());
 		UE_TRACE_LOG(MassTrace, MassBulkEntityDestroyed, MassChannel)
 			<< MassBulkEntityDestroyed.Cycle(Cycle)
+			<< MassBulkEntityDestroyed.RecordingTime(UE::Mass::Trace::GMassRecordingTime)
 			<< MassBulkEntityDestroyed.Entities(EntitiesAsU64.GetData(), EntitiesAsU64.Num());
 	}
 }
 
+void FMassTrace::SetRecordingTime(const double InRecordingTime)
+{
+	UE::Mass::Trace::GMassRecordingTime = InRecordingTime;
+}
+
 uint64 FMassTrace::OutputBeginPhaseWithID(const TCHAR* PhaseName)
 {
 	if (UE_TRACE_CHANNELEXPR_IS_ENABLED(MassChannel))
diff --git a/Engine/Source/Runtime/MassEntity/Private/MassProcessingPhaseManager.cpp b/Engine/Source/Runtime/MassEntity/Private/MassProcessingPhaseManager.cpp
--- a/Engine/Source/Runtime/MassEntity/Private/MassProcessingPhaseManager.cpp
+++ b/Engine/Source/Runtime/MassEntity/Private/MassProcessingPhaseManager.cpp
@@ -14,6 +14,10 @@
 #include "MassProcessingContext.h"
 #include "Misc/StringOutputDevice.h"
 #include "MassProcessingQueue.h"
+#include "ObjectTraceDefines.h"
+#if OBJECT_TRACE_ENABLED
+#include "ObjectTrace.h"
+#endif
 
 #include UE_INLINE_GENERATED_CPP_BY_NAME(MassProcessingPhaseManager)
 
@@ -568,6 +572,10 @@ void FMassProcessingPhaseManager::OnPhaseStart(FMassProcessingPhase& Phase)
 		HandlePendingDynamicProcessorOperations(PhaseAsInt);
 	}
 
+#if OBJECT_TRACE_ENABLED
+	UE_TRACE_MASS_SET_RECORDING_TIME(FObjectTrace::GetWorldElapsedTime(EntityManager->GetWorld()))
+#endif
+
 	UE_TRACE_MASS_PHASE_BEGIN(PhaseAsInt)
 
 	if (Owner.IsValid()
diff --git a/Engine/Source/Runtime/MassEntity/Public/MassEntityTrace.h b/Engine/Source/Runtime/MassEntity/Public/MassEntityTrace.h
--- a/Engine/Source/Runtime/MassEntity/Public/MassEntityTrace.h
+++ b/Engine/Source/Runtime/MassEntity/Public/MassEntityTrace.h
@@ -38,6 +38,12 @@ public:
 		int32 EntityCount = 0;
 	};
 
+	/** Set the current world elapsed time (recording time) for subsequent trace events.
+	 *  Should be called once per frame/phase from code that has UWorld access.
+	 *  This allows Mass trace events to be aligned with the RewindDebugger's
+	 *  ElapsedTime-based scrub time, which pauses when the simulation pauses. */
+	MASSENTITY_API static void SetRecordingTime(double InRecordingTime);
+
 	MASSENTITY_API static uint64 OutputBeginPhaseWithID(const TCHAR* PhaseName);
 	MASSENTITY_API static void OutputBeginPhaseRegion(const TCHAR* PhaseName);
 	MASSENTITY_API static void OutputEndPhaseRegion(const TCHAR* PhaseName);
@@ -45,14 +51,14 @@ public:
 
 	MASSENTITY_API static void OnPhaseBegin(uint64 PhaseId);
 	MASSENTITY_API static void OnPhaseEnd(uint64 PhaseId);
-	
+
 	MASSENTITY_API static uint64 OutputRegisterArchetype(uint64 ArchetypeID, const FMassArchetypeCompositionDescriptor& CompositionDescriptor);
 	MASSENTITY_API static uint64 RegisterArchetype(const FMassArchetypeHandle& ArchetypeHandle);
 	MASSENTITY_API static uint64 RegisterArchetype(const FMassArchetypeData& Data);
 
 	MASSENTITY_API static void OutputRegisterFragment(const UScriptStruct* Struct);
 	MASSENTITY_API static void RegisterFragment(const UScriptStruct* Struct);
-	
+
 	MASSENTITY_API static void EntityCreated(FMassEntityHandle Entity, const FMassArchetypeData& Archetype);
 	MASSENTITY_API static void EntityMoved(FMassEntityHandle Entity, const FMassArchetypeData& NewArchetype);
 	MASSENTITY_API static void EntityDestroyed(FMassEntityHandle Entity);
@@ -64,12 +70,6 @@ public:
 	MASSENTITY_API static void QueryArchetypeAdded(const FMassEntityQuery* Query, const FMassArchetypeHandle& Archetype);
 };
 
-#define UE_TRACE_MASS_PHASE_BEGIN(PhaseID) \
-	FMassTrace::OnPhaseBegin(PhaseID);
-
-#define UE_TRACE_MASS_PHASE_END(PhaseID) \
-	FMassTrace::OnPhaseEnd(PhaseID);
-
 #define UE_TRACE_MASS_ARCHETYPE_CREATED(Archetype) \
 	FMassTrace::RegisterArchetype(Archetype);
 
@@ -99,6 +99,11 @@ public:
 	{ \
 		FMassTrace::EntityDestroyed(Entity); \
 	}
+#define UE_TRACE_MASS_PHASE_BEGIN(PhaseID) \
+	FMassTrace::OnPhaseBegin(PhaseID);
+
+#define UE_TRACE_MASS_PHASE_END(PhaseID) \
+	FMassTrace::OnPhaseEnd(PhaseID);
 
 #define UE_TRACE_MASS_QUERY_CREATED() \
 	FMassTrace::QueryCreated(this);
@@ -112,6 +117,9 @@ public:
 #define UE_TRACE_MASS_QUERY_ARCHETYPE_ADDED(Archetype) \
 	FMassTrace::QueryArchetypeAdded(this, Archetype);
 
+#define UE_TRACE_MASS_SET_RECORDING_TIME(RecordingTime) \
+	FMassTrace::SetRecordingTime(RecordingTime);
+
 #define UE_TRACE_SCOPED_MASS_QUERY_FOR_EACH() \
 	FMassTrace::FScopedQueryForEachTrace _ScopedQueryForEachTrace(this);
 
@@ -133,6 +141,7 @@ public:
 #define UE_TRACE_MASS_QUERY_DESTROYED()
 #define UE_TRACE_MASS_QUERY_REGISTERED_TO_PROCESSOR(Processor)
 #define UE_TRACE_MASS_QUERY_ARCHETYPE_ADDED(Archetype)
+#define UE_TRACE_MASS_SET_RECORDING_TIME(RecordingTime)
 #define UE_TRACE_SCOPED_MASS_QUERY_FOR_EACH()
 #define UE_TRACE_SCOPED_MASS_QUERY_FOR_EACH_REPORT_ARCHETYPE(Archetype)
 
