[MassEntity] Sortable columns and step columns in RewindDebugger details view

Add clickable column headers to the Mass Entity details view so users can sort fragments by name, type, size, or step number.

Move step attribution from inline [step N] text into dedicated PrevStep and CurrStep columns.
These columns are only visible when there are merged same-frame archetype changes with step data.

#rnx

#Synced-CL 52460709

diff --git a/Engine/Source/Editor/MassEntityEditor/Private/Trace/SMassEntityDetailsView.cpp b/Engine/Source/Editor/MassEntityEditor/Private/Trace/SMassEntityDetailsView.cpp
--- a/Engine/Source/Editor/MassEntityEditor/Private/Trace/SMassEntityDetailsView.cpp
+++ b/Engine/Source/Editor/MassEntityEditor/Private/Trace/SMassEntityDetailsView.cpp
@@ -21,7 +21,9 @@ static constexpr FLinearColor ReaddedColor(0.6f, 0.6f, 0.9f, 1.0f);
 static constexpr FLinearColor InfoColor(0.5f, 0.5f, 0.5f, 1.0f);
 
 static const FName PreviousColumnName(TEXT("Previous"));
+static const FName PreviousStepColumnName(TEXT("PrevStep"));
 static const FName CurrentColumnName(TEXT("Current"));
+static const FName CurrentStepColumnName(TEXT("CurrStep"));
 static const FName TypeColumnName(TEXT("Type"));
 static const FName SizeColumnName(TEXT("Size"));
 
@@ -63,7 +65,7 @@ public:
 			if (Entry->Status == EDiffStatus::Removed)
 			{
 				return SNew(STextBlock)
-					.Text(FormatWithStep(LOCTEXT("PrefixRemove", "-"), FragName, Entry->StepIndex))
+					.Text(FText::Format(LOCTEXT("PrefixRemove", "- {0}"), FragName))
 					.ColorAndOpacity(RemovedColor);
 			}
 			// Show fragment name if it was in the previous archetype
@@ -76,6 +78,17 @@ public:
 			return SNullWidget::NullWidget;
 		}
 
+		if (ColumnName == PreviousStepColumnName)
+		{
+			if (Entry->RemovedAtStep > 0)
+			{
+				return SNew(STextBlock)
+					.Text(FText::AsNumber(Entry->RemovedAtStep))
+					.ColorAndOpacity(RemovedColor);
+			}
+			return SNullWidget::NullWidget;
+		}
+
 		if (ColumnName == CurrentColumnName)
 		{
 			// Removed elements are shown in Previous column, not here
@@ -94,11 +107,11 @@ public:
 				break;
 			case EDiffStatus::Added:
 				Color = AddedColor;
-				DisplayText = FormatWithStep(LOCTEXT("PrefixAdd", "+"), FragName, Entry->StepIndex);
+				DisplayText = FText::Format(LOCTEXT("PrefixAdd", "+ {0}"), FragName);
 				break;
 			case EDiffStatus::Readded:
 				Color = ReaddedColor;
-				DisplayText = FormatWithStep(LOCTEXT("PrefixReadd", "~"), FragName, Entry->StepIndex);
+				DisplayText = FText::Format(LOCTEXT("PrefixReadd", "~ {0}"), FragName);
 				break;
 			default:
 				break;
@@ -109,6 +122,18 @@ public:
 				.ColorAndOpacity(Color);
 		}
 
+		if (ColumnName == CurrentStepColumnName)
+		{
+			if (Entry->AddedAtStep > 0)
+			{
+				const FLinearColor Color = Entry->Status == EDiffStatus::Readded ? ReaddedColor : AddedColor;
+				return SNew(STextBlock)
+					.Text(FText::AsNumber(Entry->AddedAtStep))
+					.ColorAndOpacity(Color);
+			}
+			return SNullWidget::NullWidget;
+		}
+
 		if (ColumnName == TypeColumnName)
 		{
 			return SNew(STextBlock)
@@ -127,15 +152,6 @@ public:
 	}
 
 private:
-	static FText FormatWithStep(const FText& Prefix, const FText& Name, const int32 StepIndex)
-	{
-		if (StepIndex > 0)
-		{
-			return FText::Format(LOCTEXT("StepFmt", "{0} {1} [step {2}]"), Prefix, Name, FText::AsNumber(StepIndex));
-		}
-		return FText::Format(LOCTEXT("NoStepFmt", "{0} {1}"), Prefix, Name);
-	}
-
 	FText GetTypeText() const
 	{
 		using namespace UE::Mass::Trace;
@@ -172,22 +188,44 @@ void SMassEntityDetailsView::Construct(const FArguments& InArgs)
 	EntityId = InArgs._EntityId;
 	Cache = InArgs._Cache;
 
+	SortColumn = CurrentColumnName;
+
 	HeaderRow = SNew(SHeaderRow)
 		+ SHeaderRow::Column(PreviousColumnName)
 		.DefaultLabel(LOCTEXT("ColPrev", "Previous"))
 		.FillWidth(1.0f)
+		.SortMode(this, &SMassEntityDetailsView::GetSortModeForColumn, PreviousColumnName)
+		.OnSort(this, &SMassEntityDetailsView::OnSortModeChanged)
+
+		+ SHeaderRow::Column(PreviousStepColumnName)
+		.DefaultLabel(LOCTEXT("ColPrevStep", "Step"))
+		.FillWidth(0.2f)
+		.SortMode(this, &SMassEntityDetailsView::GetSortModeForColumn, PreviousStepColumnName)
+		.OnSort(this, &SMassEntityDetailsView::OnSortModeChanged)
 
 		+ SHeaderRow::Column(CurrentColumnName)
 		.DefaultLabel(LOCTEXT("ColCurr", "Current"))
 		.FillWidth(1.0f)
+		.SortMode(this, &SMassEntityDetailsView::GetSortModeForColumn, CurrentColumnName)
+		.OnSort(this, &SMassEntityDetailsView::OnSortModeChanged)
+
+		+ SHeaderRow::Column(CurrentStepColumnName)
+		.DefaultLabel(LOCTEXT("ColCurrStep", "Step"))
+		.FillWidth(0.2f)
+		.SortMode(this, &SMassEntityDetailsView::GetSortModeForColumn, CurrentStepColumnName)
+		.OnSort(this, &SMassEntityDetailsView::OnSortModeChanged)
 
 		+ SHeaderRow::Column(TypeColumnName)
 		.DefaultLabel(LOCTEXT("ColType", "Type"))
 		.FillWidth(0.5f)
+		.SortMode(this, &SMassEntityDetailsView::GetSortModeForColumn, TypeColumnName)
+		.OnSort(this, &SMassEntityDetailsView::OnSortModeChanged)
 
 		+ SHeaderRow::Column(SizeColumnName)
 		.DefaultLabel(LOCTEXT("ColSize", "Size"))
-		.FillWidth(0.3f);
+		.FillWidth(0.3f)
+		.SortMode(this, &SMassEntityDetailsView::GetSortModeForColumn, SizeColumnName)
+		.OnSort(this, &SMassEntityDetailsView::OnSortModeChanged);
 
 	ChildSlot
 	[
@@ -350,6 +388,11 @@ void SMassEntityDetailsView::RebuildFragmentList()
 	// Show/hide Previous column based on whether there is a previous archetype
 	HeaderRow->SetShowGeneratedColumn(UE::MassEntityDetailsView::Private::PreviousColumnName, PreviousArchetypeId != 0);
 
+	// Show step columns only when there are merged same-frame changes with step data
+	const bool bHasSteps = LastDisplayedMergeCount > 1 && LastDisplayedArchetypeBeforeBatch != 0;
+	HeaderRow->SetShowGeneratedColumn(UE::MassEntityDetailsView::Private::PreviousStepColumnName, bHasSteps);
+	HeaderRow->SetShowGeneratedColumn(UE::MassEntityDetailsView::Private::CurrentStepColumnName, bHasSteps);
+
 	if (ArchetypeAtScrubTime == 0)
 	{
 		if (ListView.IsValid())
@@ -504,12 +547,13 @@ void SMassEntityDetailsView::RebuildFragmentList()
 		if (bIsReadded)
 		{
 			Entry->Status = EDiffStatus::Readded;
-			Entry->StepIndex = RemovedAtStep.FindRef(Fragment->Id);
+			Entry->AddedAtStep = AddedAtStep.FindRef(Fragment->Id);
+			Entry->RemovedAtStep = RemovedAtStep.FindRef(Fragment->Id);
 		}
 		else if (!bInPrevious && PreviousArchetype)
 		{
 			Entry->Status = EDiffStatus::Added;
-			Entry->StepIndex = AddedAtStep.FindRef(Fragment->Id);
+			Entry->AddedAtStep = AddedAtStep.FindRef(Fragment->Id);
 		}
 		else
 		{
@@ -532,17 +576,30 @@ void SMassEntityDetailsView::RebuildFragmentList()
 			TSharedPtr<FFragmentEntry> Entry = MakeShared<FFragmentEntry>();
 			Entry->Info = Fragment;
 			Entry->Status = EDiffStatus::Removed;
-			Entry->StepIndex = RemovedAtStep.FindRef(Fragment->Id);
+			Entry->RemovedAtStep = RemovedAtStep.FindRef(Fragment->Id);
 			Entry->bHasPrevious = true;
 			FragmentEntries.Add(MoveTemp(Entry));
 		}
 	}
 
-	// Sort alphabetically
-	FragmentEntries.Sort([](const TSharedPtr<FFragmentEntry>& A, const TSharedPtr<FFragmentEntry>& B)
+	SortFragmentEntries();
+
+	if (ListView.IsValid())
 	{
-		return A->Info->Name < B->Info->Name;
-	});
+		ListView->RequestListRefresh();
+	}
+}
+
+EColumnSortMode::Type SMassEntityDetailsView::GetSortModeForColumn(const FName ColumnId) const
+{
+	return ColumnId == SortColumn ? SortMode : EColumnSortMode::None;
+}
+
+void SMassEntityDetailsView::OnSortModeChanged(const EColumnSortPriority::Type /*Priority*/, const FName& ColumnId, const EColumnSortMode::Type InSortMode)
+{
+	SortColumn = ColumnId;
+	SortMode = InSortMode;
+	SortFragmentEntries();
 
 	if (ListView.IsValid())
 	{
@@ -550,4 +607,41 @@ void SMassEntityDetailsView::RebuildFragmentList()
 	}
 }
 
+void SMassEntityDetailsView::SortFragmentEntries()
+{
+	using namespace UE::MassEntityDetailsView::Private;
+
+	const bool bAscending = SortMode != EColumnSortMode::Descending;
+
+	FragmentEntries.Sort([this, bAscending](const TSharedPtr<FFragmentEntry>& A, const TSharedPtr<FFragmentEntry>& B)
+	{
+		int32 Cmp = 0;
+
+		if (SortColumn == TypeColumnName)
+		{
+			Cmp = static_cast<int32>(A->Info->Type) - static_cast<int32>(B->Info->Type);
+		}
+		else if (SortColumn == SizeColumnName)
+		{
+			Cmp = static_cast<int32>(A->Info->Size) - static_cast<int32>(B->Info->Size);
+		}
+		else if (SortColumn == PreviousStepColumnName)
+		{
+			Cmp = A->RemovedAtStep - B->RemovedAtStep;
+		}
+		else if (SortColumn == CurrentStepColumnName)
+		{
+			Cmp = A->AddedAtStep - B->AddedAtStep;
+		}
+
+		// Default / tie-break: sort by fragment name
+		if (Cmp == 0)
+		{
+			Cmp = A->Info->Name.Compare(B->Info->Name);
+		}
+
+		return bAscending ? Cmp < 0 : Cmp > 0;
+	});
+}
+
 #undef LOCTEXT_NAMESPACE
diff --git a/Engine/Source/Editor/MassEntityEditor/Private/Trace/SMassEntityDetailsView.h b/Engine/Source/Editor/MassEntityEditor/Private/Trace/SMassEntityDetailsView.h
--- a/Engine/Source/Editor/MassEntityEditor/Private/Trace/SMassEntityDetailsView.h
+++ b/Engine/Source/Editor/MassEntityEditor/Private/Trace/SMassEntityDetailsView.h
@@ -4,11 +4,11 @@
 
 #include "Widgets/SCompoundWidget.h"
 #include "Widgets/DeclarativeSyntaxSupport.h"
+#include "Widgets/Views/SHeaderRow.h"
 
 class STextBlock;
 class STableViewBase;
 class ITableRow;
-class SHeaderRow;
 template <typename ItemType> class SListView;
 
 namespace UE::Mass::Trace
@@ -36,7 +36,10 @@ public:
 	{
 		const UE::Mass::Trace::FFragmentInfo* Info = nullptr;
 		EDiffStatus Status = EDiffStatus::Unchanged;
-		int32 StepIndex = 0;
+		/** Step at which the fragment was added (Added/Readded). 0 = not applicable. */
+		int32 AddedAtStep = 0;
+		/** Step at which the fragment was removed (Removed/Readded). 0 = not applicable. */
+		int32 RemovedAtStep = 0;
 		bool bHasPrevious = false;
 	};
 
@@ -62,8 +65,12 @@ private:
 	};
 
 	void RebuildFragmentList();
+	void SortFragmentEntries();
 	FArchetypeResolution ResolveArchetypesAtScrubTime() const;
 
+	EColumnSortMode::Type GetSortModeForColumn(FName ColumnId) const;
+	void OnSortModeChanged(EColumnSortPriority::Type Priority, const FName& ColumnId, EColumnSortMode::Type InSortMode);
+
 	TSharedRef<ITableRow> OnGenerateRow(TSharedPtr<FFragmentEntry> Entry, const TSharedRef<STableViewBase>& OwnerTable);
 
 	uint64 EntityId = 0;
@@ -79,4 +86,7 @@ private:
 	TSharedPtr<SListView<TSharedPtr<FFragmentEntry>>> ListView;
 	TSharedPtr<SHeaderRow> HeaderRow;
 	TSharedPtr<STextBlock> MergeInfoText;
+
+	FName SortColumn;
+	EColumnSortMode::Type SortMode = EColumnSortMode::Ascending;
 };
